Copying //UE4/Dev-Physics to Dev-Main (//UE4/Dev-Main) [at] 9098278

#rnx

diff --git a/Engine/Config/BaseEngine.ini b/Engine/Config/BaseEngine.ini
--- a/Engine/Config/BaseEngine.ini
+++ b/Engine/Config/BaseEngine.ini
@@ -80,6 +80,8 @@ UnrealEdEngine=/Script/UnrealEd.UnrealEdEngine
 ;LargeFontName=/Engine/EngineFonts/Roboto.Roboto
 ;SubtitleFontName=/Engine/EngineFonts/Roboto.Roboto
 WireframeMaterialName=/Engine/EngineDebugMaterials/WireframeMaterial.WireframeMaterial
+HairDefaultMaterialName=/Engine/EngineMaterials/HairDefaultMaterial.HairDefaultMaterial
+HairDebugMaterialName=/Engine/EngineMaterials/HairDebugMaterial.HairDebugMaterial
 DefaultMaterialName=/Engine/EngineMaterials/WorldGridMaterial.WorldGridMaterial
 DefaultLightFunctionMaterialName=/Engine/EngineMaterials/DefaultLightFunctionMaterial.DefaultLightFunctionMaterial
 DefaultTextureName=/Engine/EngineResources/DefaultTexture.DefaultTexture
@@ -691,6 +693,15 @@ bUseStreamingPause=false
 ; 4.24
 +PropertyRedirects=(OldName="NavDataConfig.NavigationDataClassName", NewName="NavDataConfig.NavDataClass")
 
+; Rename ClothingSystemRuntime to ClothingSystemRuntimeNv
++PackageRedirects=(OldName="/Script/ClothingSystemRuntime",NewName="/Script/ClothingSystemRuntimeNv")
++ClassRedirects=(OldName="/Script/ClothingSystemRuntime.ClothingSimulationFactoryNv",NewName="/Script/ClothingSystemRuntimeNv.ClothingSimulationFactoryNv")
++ClassRedirects=(OldName="/Script/ClothingSystemRuntime.ClothingSimulationInteractorNv",NewName="/Script/ClothingSystemRuntimeNv.ClothingSimulationInteractorNv")
++FunctionRedirects=(OldName="ClothingSystemRuntime.ClothingSimulationInteractorNv.SetAnimDriveSpringStiffness",NewName="ClothingSystemRuntimeNv.ClothingSimulationInteractorNv.SetAnimDriveSpringStiffness")
++FunctionRedirects=(OldName="ClothingSystemRuntime.ClothingSimulationInteractorNv.SetAnimDriveDamperStiffness",NewName="ClothingSystemRuntimeNv.ClothingSimulationInteractorNv.SetAnimDriveDamperStiffness")
++FunctionRedirects=(OldName="ClothingSystemRuntime.ClothingSimulationInteractorNv.EnableGravityOverride",NewName="ClothingSystemRuntimeNv.ClothingSimulationInteractorNv.EnableGravityOverride")
++FunctionRedirects=(OldName="ClothingSystemRuntime.ClothingSimulationInteractorNv.DisableGravityOverride",NewName="ClothingSystemRuntimeNv.ClothingSimulationInteractorNv.DisableGravityOverride")
+
 +FunctionRedirects=(OldName="Actor.GetComponentsByClass", NewName="Actor.K2_GetComponentsByClass")
 
 +FunctionRedirects=(OldName="SkeletalMeshComponent.GetSubInstanceByTag",NewName="SkeletalMeshComponent.GetLinkedAnimGraphInstanceByTag")
@@ -802,6 +813,8 @@ bFullyCompressStartupPackages=false
 +Package=/Engine/EngineDebugMaterials/VertexColorViewMode_GreenOnly
 +Package=/Engine/EngineDebugMaterials/VertexColorViewMode_RedOnly
 +Package=/Engine/EngineDebugMaterials/WireframeMaterial
++Package=/Engine/EngineDebugMaterials/HairDefaultMaterial
++Package=/Engine/EngineDebugMaterials/HairDebugMaterial
 
 +Package=/Engine/EngineSounds/WhiteNoise
 
@@ -1339,7 +1352,7 @@ bHasVoiceEnabled=true
 ; Internal
 VoiceNotificationDelta=0.33
 ; Steam
-;VoiceNotificationDelta=0.2 
+;VoiceNotificationDelta=0.2
 MaxLocalTalkers=1
 MaxRemoteTalkers=16
 bUseBuildIdOverride=false
@@ -1532,7 +1545,7 @@ HighDetail=(NumLODs=6,MaxNumStreamedLODs=0,bSupportLODStreaming=0,LODPercentTria
 
 ; Config for RuntimeAssetCache asset groups.
 ; Each group can be managed separately (e.g. advertisement cache, character image cache etc.)
-; For each group a name must be specified. Size is optional and defaults to 5MB. Size in 
+; For each group a name must be specified. Size is optional and defaults to 5MB. Size in
 ; ini file is stored as number of bytes. E.g.:
 ; [RuntimeAssetCache]
 ; +BucketConfigs=(Name="CharacterImage", Size=3000000)
diff --git a/Engine/Plugins/Experimental/AlembicHairImporter/AlembicHairImporter.uplugin b/Engine/Plugins/Experimental/AlembicHairImporter/AlembicHairImporter.uplugin
new file mode 100644
--- /dev/null
+++ b/Engine/Plugins/Experimental/AlembicHairImporter/AlembicHairImporter.uplugin
@@ -0,0 +1,34 @@
+{
+	"FileVersion": 1,
+	"Version": 1,
+	"VersionName": "0.1",
+	"FriendlyName": "Alembic Hair Strands Importer",
+	"Description": "Import Hair Strands from Alembic file",
+	"Category": "Geometry",
+	"CreatedBy": "Epic Games, Inc.",
+	"CreatedByURL": "http://epicgames.com",
+	"DocsURL": "",
+	"MarketplaceURL": "",
+	"SupportURL": "",
+	"EnabledByDefault": false,
+	"CanContainContent": false,
+	"IsBetaVersion": true,
+	"Installed": false,
+	"Modules": [
+		{
+			"Name": "AlembicHairTranslatorModule",
+			"Type": "Editor",
+			"LoadingPhase": "Default",
+			"WhitelistPlatforms": [
+				"Win64",
+				"Win32"
+			]
+		}
+	],
+	"Plugins": [
+		{
+			"Name": "HairStrands",
+			"Enabled": true
+		}
+	]
+}
diff --git a/Engine/Plugins/Experimental/AlembicHairImporter/Source/AlembicHairTranslator/Private/AlembicHairTranslator.cpp b/Engine/Plugins/Experimental/AlembicHairImporter/Source/AlembicHairTranslator/Private/AlembicHairTranslator.cpp
new file mode 100644
--- /dev/null
+++ b/Engine/Plugins/Experimental/AlembicHairImporter/Source/AlembicHairTranslator/Private/AlembicHairTranslator.cpp
@@ -0,0 +1,557 @@
+// Copyright 1998-2019 Epic Games, Inc. All Rights Reserved.
+
+#include "AlembicHairTranslator.h"
+
+#include "HairDescription.h"
+#include "Misc/Paths.h"
+
+#if PLATFORM_WINDOWS
+#include "Windows/WindowsHWrapper.h"
+#include "Windows/AllowWindowsPlatformTypes.h"
+#endif
+
+THIRD_PARTY_INCLUDES_START
+#include "Alembic/AbcGeom/All.h"
+#include "Alembic/AbcCoreFactory/IFactory.h"
+#include "Alembic/Abc/IArchive.h"
+#include "Alembic/Abc/IObject.h"
+THIRD_PARTY_INCLUDES_END
+
+#if PLATFORM_WINDOWS
+#include "Windows/HideWindowsPlatformTypes.h"
+#endif
+
+DEFINE_LOG_CATEGORY_STATIC(LogAlembicHairImporter, Log, All);
+
+namespace AlembicHairFormat
+{
+	static const float RootRadius = 0.0001f; // m
+	static const float TipRadius = 0.00005f; // m
+
+	static constexpr const float UNIT_TO_CM = 1;
+}
+
+namespace AlembicHairTranslatorUtils
+{
+	bool IsAttributeValid(const FString& AttributeName)
+	{
+		// Ignore the attributes that aren't prefixed with groom_
+		return AttributeName.StartsWith(TEXT("groom_"));
+	}
+
+	template <typename AbcParamType, typename AbcArraySampleType, typename AttributeType>
+	void SetGroomAttributes(FHairDescription& HairDescription, const Alembic::AbcGeom::ICompoundProperty& Parameters, const std::string& PropName)
+	{
+		FName AttributeName(ANSI_TO_TCHAR(PropName.c_str()));
+
+		AbcParamType Param(Parameters, PropName);
+		AbcArraySampleType ParamValues = Param.getExpandedValue().getVals();
+		if (ParamValues->size() == 1)
+		{
+			AttributeType ParamValue = (*ParamValues)[0];
+			SetGroomAttribute(HairDescription, FGroomID(0), AttributeName, ParamValue);
+		}
+	}
+
+	template <>
+	void SetGroomAttributes<Alembic::AbcGeom::IStringGeomParam, Alembic::Abc::StringArraySamplePtr, FName>(FHairDescription& HairDescription, const Alembic::AbcGeom::ICompoundProperty& Parameters, const std::string& PropName)
+	{
+		FName AttributeName(ANSI_TO_TCHAR(PropName.c_str()));
+
+		Alembic::AbcGeom::IStringGeomParam Param(Parameters, PropName);
+		Alembic::Abc::StringArraySamplePtr ParamValues = Param.getExpandedValue().getVals();
+		if (ParamValues->size() == 1)
+		{
+			FName ParamValue = FName((*ParamValues)[0].c_str());
+			SetGroomAttribute(HairDescription, FGroomID(0), AttributeName, ParamValue);
+		}
+	}
+
+	template <typename AbcParamType, typename AbcArraySampleType, typename AttributeType>
+	void SetGroomAttributes(FHairDescription& HairDescription, const Alembic::AbcGeom::ICompoundProperty& Parameters, const std::string& PropName, uint8 Extent)
+	{
+		FName AttributeName(ANSI_TO_TCHAR(PropName.c_str()));
+
+		AbcParamType Param(Parameters, PropName);
+		AbcArraySampleType ParamValues = Param.getExpandedValue().getVals();
+		if (ParamValues->size() == 1)
+		{
+			AttributeType ParamValue;
+			for (int32 Index = 0; Index < Extent; ++Index)
+			{
+				ParamValue[Index] = (*ParamValues)[0][Index];
+			}
+			SetGroomAttribute(HairDescription, FGroomID(0), AttributeName, ParamValue);
+		}
+	}
+
+	void SetGroomAttributes(FHairDescription& HairDescription, const Alembic::AbcGeom::ICompoundProperty& Parameters)
+	{
+		for (int Index = 0; Index < Parameters.getNumProperties(); ++Index)
+		{
+			Alembic::Abc::PropertyHeader PropertyHeader = Parameters.getPropertyHeader(Index);
+			std::string PropName = PropertyHeader.getName();
+
+			if (!IsAttributeValid(ANSI_TO_TCHAR(PropName.c_str())))
+			{
+				continue;
+			}
+
+			Alembic::Abc::PropertyType PropType = PropertyHeader.getPropertyType();
+
+			if (PropType != Alembic::Abc::kCompoundProperty)
+			{
+				Alembic::Abc::DataType DataType = PropertyHeader.getDataType();
+				uint8 Extent = DataType.getExtent();
+
+				switch (DataType.getPod())
+				{
+				case Alembic::Util::kInt16POD:
+				{
+					SetGroomAttributes<Alembic::AbcGeom::IInt16GeomParam, Alembic::Abc::Int16ArraySamplePtr, int>(HairDescription, Parameters, PropName);
+				}
+				break;
+				case Alembic::Util::kInt32POD:
+				{
+					SetGroomAttributes<Alembic::AbcGeom::IInt32GeomParam, Alembic::Abc::Int32ArraySamplePtr, int>(HairDescription, Parameters, PropName);
+				}
+				break;
+				case Alembic::Util::kStringPOD:
+				{
+					SetGroomAttributes<Alembic::AbcGeom::IStringGeomParam, Alembic::Abc::StringArraySamplePtr, FName>(HairDescription, Parameters, PropName);
+				}
+				break;
+				case Alembic::Util::kFloat32POD:
+				{
+					switch (Extent)
+					{
+					case 1:
+						SetGroomAttributes<Alembic::AbcGeom::IFloatGeomParam, Alembic::Abc::FloatArraySamplePtr, float>(HairDescription, Parameters, PropName);
+						break;
+					case 2:
+						SetGroomAttributes<Alembic::AbcGeom::IV2fGeomParam, Alembic::Abc::V2fArraySamplePtr, FVector2D>(HairDescription, Parameters, PropName, Extent);
+						break;
+					case 3:
+						SetGroomAttributes<Alembic::AbcGeom::IV3fGeomParam, Alembic::Abc::V3fArraySamplePtr, FVector>(HairDescription, Parameters, PropName, Extent);
+						break;
+					}
+				}
+				break;
+				case Alembic::Util::kFloat64POD:
+				{
+					switch (Extent)
+					{
+					case 1:
+						SetGroomAttributes<Alembic::AbcGeom::IDoubleGeomParam, Alembic::Abc::DoubleArraySamplePtr, float>(HairDescription, Parameters, PropName);
+						break;
+					case 2:
+						SetGroomAttributes<Alembic::AbcGeom::IV2dGeomParam, Alembic::Abc::V2dArraySamplePtr, FVector2D>(HairDescription, Parameters, PropName, Extent);
+						break;
+					case 3:
+						SetGroomAttributes<Alembic::AbcGeom::IV3dGeomParam, Alembic::Abc::V3dArraySamplePtr, FVector>(HairDescription, Parameters, PropName, Extent);
+						break;
+					}
+				}
+				break;
+				}
+			}
+		}
+	}
+
+	template <typename AbcParamType, typename AbcArraySampleType, typename AttributeType>
+	void ConvertAlembicAttribute(FHairDescription& HairDescription, FStrandID StrandID, int32 StartVertexID, int32 NumVertices, const Alembic::AbcGeom::ICompoundProperty& Parameters, const std::string& PropName)
+	{
+		FName AttributeName(ANSI_TO_TCHAR(PropName.c_str()));
+
+		// The number of values in Param determines the scope on which to set the hair attribute
+		AbcParamType Param(Parameters, PropName);
+		AbcArraySampleType ParamValues = Param.getExpandedValue().getVals();
+		if (ParamValues->size() == 1)
+		{
+			AttributeType ParamValue = (*ParamValues)[0];
+			SetHairStrandAttribute(HairDescription, StrandID, AttributeName, ParamValue);
+
+		}
+		else if (ParamValues->size() == NumVertices)
+		{
+			TVertexAttributesRef<AttributeType> VertexAttributeRef = HairDescription.VertexAttributes().GetAttributesRef<AttributeType>(AttributeName);
+			if (!VertexAttributeRef.IsValid())
+			{
+				HairDescription.VertexAttributes().RegisterAttribute<AttributeType>(AttributeName);
+				VertexAttributeRef = HairDescription.VertexAttributes().GetAttributesRef<AttributeType>(AttributeName);
+			}
+
+			for (int32 VertexIndex = 0; VertexIndex < NumVertices; ++VertexIndex)
+			{
+				VertexAttributeRef[FVertexID(StartVertexID + VertexIndex)] = (*ParamValues)[VertexIndex];
+			}
+		}
+	}
+
+	template <typename AbcParamType, typename AbcArraySampleType, typename AttributeType>
+	void ConvertAlembicAttribute(FHairDescription& HairDescription, FStrandID StrandID, int32 StartVertexID, int32 NumVertices, const Alembic::AbcGeom::ICompoundProperty& Parameters, const std::string& PropName, uint8 Extent)
+	{
+		FName AttributeName(ANSI_TO_TCHAR(PropName.c_str()));
+
+		// The number of values in Param determines the scope on which to set the hair attribute
+		AbcParamType Param(Parameters, PropName);
+		AbcArraySampleType ParamValues = Param.getExpandedValue().getVals();
+		if (ParamValues->size() == 1)
+		{
+			AttributeType ParamValue;
+			for (int32 Index = 0; Index < Extent; ++Index)
+			{
+				ParamValue[Index] = (*ParamValues)[0][Index];
+			}
+			SetHairStrandAttribute(HairDescription, StrandID, AttributeName, ParamValue);
+		}
+		else if (ParamValues->size() == NumVertices)
+		{
+			TVertexAttributesRef<AttributeType> VertexAttributeRef = HairDescription.VertexAttributes().GetAttributesRef<AttributeType>(AttributeName);
+			if (!VertexAttributeRef.IsValid())
+			{
+				HairDescription.VertexAttributes().RegisterAttribute<AttributeType>(AttributeName, 1, AttributeType::ZeroVector);
+				VertexAttributeRef = HairDescription.VertexAttributes().GetAttributesRef<AttributeType>(AttributeName);
+			}
+
+			for (int32 VertexIndex = 0; VertexIndex < NumVertices; ++VertexIndex)
+			{
+				AttributeType ParamValue;
+				for (int32 Index = 0; Index < Extent; ++Index)
+				{
+					ParamValue[Index] = (*ParamValues)[VertexIndex][Index];
+				}
+
+				VertexAttributeRef[FVertexID(StartVertexID + VertexIndex)] = ParamValue;
+			}
+		}
+	}
+
+	/** Convert the given Alembic parameters to hair attributes in the proper scope */
+	void ConvertAlembicAttributes(FHairDescription& HairDescription, FStrandID StrandID, int32 StartVertexID, int32 NumVertices, const Alembic::AbcGeom::ICompoundProperty& Parameters)
+	{
+		for (int Index = 0; Index < Parameters.getNumProperties(); ++Index)
+		{
+			Alembic::Abc::PropertyHeader PropertyHeader = Parameters.getPropertyHeader(Index);
+			std::string PropName = PropertyHeader.getName();
+
+			if (!IsAttributeValid(ANSI_TO_TCHAR(PropName.c_str())))
+			{
+				continue;
+			}
+
+			Alembic::Abc::PropertyType PropType = PropertyHeader.getPropertyType();
+
+			if (PropType != Alembic::Abc::kCompoundProperty)
+			{
+				Alembic::Abc::DataType DataType = PropertyHeader.getDataType();
+				uint8 Extent = DataType.getExtent();
+
+				switch (DataType.getPod())
+				{
+				case Alembic::Util::kBooleanPOD:
+				{
+					ConvertAlembicAttribute<Alembic::AbcGeom::IBoolGeomParam, Alembic::Abc::BoolArraySamplePtr, bool>(HairDescription, StrandID, StartVertexID, NumVertices, Parameters, PropName);
+				}
+				break;
+				case Alembic::Util::kInt32POD:
+				{
+					ConvertAlembicAttribute<Alembic::AbcGeom::IInt32GeomParam, Alembic::Abc::Int32ArraySamplePtr, int>(HairDescription, StrandID, StartVertexID, NumVertices, Parameters, PropName);
+				}
+				break;
+				case Alembic::Util::kFloat32POD:
+				{
+					switch (Extent)
+					{
+					case 1:
+						ConvertAlembicAttribute<Alembic::AbcGeom::IFloatGeomParam, Alembic::Abc::FloatArraySamplePtr, float>(HairDescription, StrandID, StartVertexID, NumVertices, Parameters, PropName);
+						break;
+					case 2:
+						ConvertAlembicAttribute<Alembic::AbcGeom::IV2fGeomParam, Alembic::Abc::V2fArraySamplePtr, FVector2D>(HairDescription, StrandID, StartVertexID, NumVertices, Parameters, PropName, DataType.getExtent());
+						break;
+					case 3:
+						ConvertAlembicAttribute<Alembic::AbcGeom::IV3fGeomParam, Alembic::Abc::V3fArraySamplePtr, FVector>(HairDescription, StrandID, StartVertexID, NumVertices, Parameters, PropName, DataType.getExtent());
+						break;
+					}
+				}
+				break;
+				case Alembic::Util::kFloat64POD:
+				{
+					switch (Extent)
+					{
+					case 1:
+						ConvertAlembicAttribute<Alembic::AbcGeom::IDoubleGeomParam, Alembic::Abc::DoubleArraySamplePtr, float>(HairDescription, StrandID, StartVertexID, NumVertices, Parameters, PropName);
+						break;
+					case 2:
+						ConvertAlembicAttribute<Alembic::AbcGeom::IV2dGeomParam, Alembic::Abc::V2dArraySamplePtr, FVector2D>(HairDescription, StrandID, StartVertexID, NumVertices, Parameters, PropName, DataType.getExtent());
+						break;
+					case 3:
+						ConvertAlembicAttribute<Alembic::AbcGeom::IV3dGeomParam, Alembic::Abc::V3dArraySamplePtr, FVector>(HairDescription, StrandID, StartVertexID, NumVertices, Parameters, PropName, DataType.getExtent());
+						break;
+					}
+				}
+				break;
+				}
+			}
+		}
+	}
+}
+
+FMatrix ConvertAlembicMatrix(const Alembic::Abc::M44d& AbcMatrix)
+{
+	FMatrix Matrix;
+	for (uint32 i = 0; i < 16; ++i)
+	{
+		Matrix.M[i >> 2][i % 4] = (float)AbcMatrix.getValue()[i];
+	}
+
+	return Matrix;
+}
+
+static void ParseObject(const Alembic::Abc::IObject& InObject, FHairDescription& HairDescription, const FMatrix& ParentMatrix, bool bCheckGroomAttributes)
+{
+	// Get MetaData info from current Alembic Object
+	const Alembic::Abc::MetaData ObjectMetaData = InObject.getMetaData();
+	const uint32 NumChildren = InObject.getNumChildren();
+
+	FMatrix LocalMatrix = ParentMatrix;
+
+	enum class EAttributeFrequency : uint8
+	{
+		None,
+		Groom,
+		Hair,
+		CV
+	};
+
+	bool bHandled = false;
+	if (Alembic::AbcGeom::ICurves::matches(ObjectMetaData))
+	{
+		Alembic::AbcGeom::ICurves Curves = Alembic::AbcGeom::ICurves(InObject, Alembic::Abc::kWrapExisting);
+		Alembic::AbcGeom::ICurves::schema_type::Sample Sample = Curves.getSchema().getValue();
+
+		Alembic::Abc::FloatArraySamplePtr Widths = Curves.getSchema().getWidthsParam() ? Curves.getSchema().getWidthsParam().getExpandedValue().getVals() : nullptr;
+		Alembic::Abc::P3fArraySamplePtr Positions = Sample.getPositions();
+		Alembic::Abc::Int32ArraySamplePtr NumVertices = Sample.getCurvesNumVertices();
+
+		const int32 NumWidths = Widths ? Widths->size() : 0;
+		const uint32 NumPoints = Positions ? Positions->size() : 0;
+		const uint32 NumCurves = NumVertices->size(); // equivalent to Sample.getNumCurves()
+
+		EAttributeFrequency WidthFrequency = EAttributeFrequency::None;
+		{
+			if (NumWidths == NumPoints)
+			{
+				WidthFrequency = EAttributeFrequency::CV;
+			}
+			else if (NumWidths == NumCurves)
+			{
+				WidthFrequency = EAttributeFrequency::Hair;
+			}
+		}
+
+		uint32 GlobalIndex = 0;
+		for (uint32 CurveIndex = 0; CurveIndex < NumCurves; ++CurveIndex)
+		{
+			const uint32 CurveNumVertices = (*NumVertices)[CurveIndex];
+
+			FStrandID StrandID = HairDescription.AddStrand();
+
+			SetHairStrandAttribute(HairDescription, StrandID, HairAttribute::Strand::VertexCount, (int) CurveNumVertices);
+
+			int32 StartVertexID = HairDescription.GetNumVertices();
+			for (uint32 PointIndex = 0; PointIndex < CurveNumVertices; ++PointIndex, ++GlobalIndex)
+			{
+				FVertexID VertexID = HairDescription.AddVertex();
+
+				Alembic::Abc::P3fArraySample::value_type Position = (*Positions)[GlobalIndex];
+
+				FVector ConvertedPosition = ParentMatrix.TransformPosition(FVector(Position.x, Position.y, Position.z)) * AlembicHairFormat::UNIT_TO_CM;
+				SetHairVertexAttribute(HairDescription, VertexID, HairAttribute::Vertex::Position, ConvertedPosition);
+
+				float Width = 0;
+				switch (WidthFrequency)
+				{
+				case EAttributeFrequency::None:
+				{
+					const float CoordU = PointIndex / static_cast<float>(CurveNumVertices - 1);
+					Width = FMath::Lerp(AlembicHairFormat::RootRadius, AlembicHairFormat::TipRadius, CoordU);
+				}
+					break;
+				case EAttributeFrequency::CV:
+					Width = (*Widths)[GlobalIndex];
+					break;
+				}
+
+				// Per-vertex widths
+				if ((WidthFrequency == EAttributeFrequency::CV || WidthFrequency == EAttributeFrequency::None))
+				{
+					SetHairVertexAttribute(HairDescription, VertexID, HairAttribute::Vertex::Width, Width);
+				}
+			}
+
+			Alembic::AbcGeom::ICompoundProperty ArbParams = Curves.getSchema().getArbGeomParams();
+			if (ArbParams)
+			{
+				AlembicHairTranslatorUtils::ConvertAlembicAttributes(HairDescription, StrandID, StartVertexID, CurveNumVertices, ArbParams);
+			}
+
+			if (WidthFrequency == EAttributeFrequency::Hair)
+			{
+				// Fallback if no per-strand or per-vertex groom_width attribute was found
+				TStrandAttributesRef<float> StrandWidths = HairDescription.StrandAttributes().GetAttributesRef<float>(HairAttribute::Strand::Width);
+				TVertexAttributesRef<float> VertexWidths = HairDescription.VertexAttributes().GetAttributesRef<float>(HairAttribute::Strand::Width);
+
+				if (!StrandWidths.IsValid() && !VertexWidths.IsValid())
+				{
+					const float Width = (*Widths)[CurveIndex];
+					SetHairStrandAttribute(HairDescription, StrandID, HairAttribute::Strand::Width, Width);
+				}
+				if (StrandWidths.IsValid() && StrandWidths[StrandID] == 0.f)
+				{
+					const float Width = (*Widths)[CurveIndex];
+					SetHairStrandAttribute(HairDescription, StrandID, HairAttribute::Strand::Width, Width);
+				}
+			}
+		}
+	}
+	else if (Alembic::AbcGeom::IXform::matches(ObjectMetaData))
+	{
+		Alembic::AbcGeom::IXform Xform = Alembic::AbcGeom::IXform(InObject, Alembic::Abc::kWrapExisting);
+		Alembic::AbcGeom::XformSample MatrixSample; 
+		Xform.getSchema().get(MatrixSample);
+
+		// The groom attributes should only be on the first IXform under the top node, no need to check for them once they are found
+		if (bCheckGroomAttributes)
+		{
+			Alembic::AbcGeom::ICompoundProperty ArbParams = Xform.getSchema().getArbGeomParams();
+			if (ArbParams)
+			{
+				if (ArbParams.getNumProperties() > 0)
+				{
+					AlembicHairTranslatorUtils::SetGroomAttributes(HairDescription, ArbParams);
+					bCheckGroomAttributes = false;
+				}
+			}
+		}
+
+		LocalMatrix =  ParentMatrix * ConvertAlembicMatrix(MatrixSample.getMatrix());
+	}
+
+	if (NumChildren > 0)
+	{
+		for (uint32 ChildIndex = 0; ChildIndex < NumChildren; ++ChildIndex)
+		{
+			ParseObject(InObject.getChild(ChildIndex), HairDescription, LocalMatrix, bCheckGroomAttributes);
+		}
+	}
+}
+
+bool FAlembicHairTranslator::Translate(const FString& FileName, FHairDescription& HairDescription)
+{
+	/** Factory used to generate objects*/
+	Alembic::AbcCoreFactory::IFactory Factory;
+	Alembic::AbcCoreFactory::IFactory::CoreType CompressionType = Alembic::AbcCoreFactory::IFactory::kUnknown;
+	/** Archive-typed ABC file */
+	Alembic::Abc::IArchive Archive;
+	/** Alembic typed root (top) object*/
+	Alembic::Abc::IObject TopObject;
+
+	Factory.setPolicy(Alembic::Abc::ErrorHandler::kThrowPolicy);
+	Factory.setOgawaNumStreams(12);
+
+	// Extract Archive and compression type from file
+	Archive = Factory.getArchive(TCHAR_TO_UTF8(*FileName), CompressionType);
+	if (!Archive.valid())
+	{
+		UE_LOG(LogAlembicHairImporter, Warning, TEXT("Failed to open %s: Not a valid Alembic file."), *FileName);
+		return false;
+	}
+
+	// Get Top/root object
+	TopObject = Alembic::Abc::IObject(Archive, Alembic::Abc::kTop);
+	if (!TopObject.valid())
+	{
+		UE_LOG(LogAlembicHairImporter, Warning, TEXT("Failed to import %s: Root not is not valid."), *FileName);
+		return false;
+	}
+
+	FMatrix ParentMatrix = FMatrix::Identity;
+	ParseObject(TopObject, HairDescription, ParentMatrix, true);
+
+	return HairDescription.IsValid();
+}
+
+static void ValidateObject(const Alembic::Abc::IObject& InObject, bool& bHasGeometry, int32& NumCurves)
+{
+	// Validate that the Alembic has curves only
+	// Any PolyMesh will cause the Alembic to be rejected by this translator
+
+	Alembic::AbcCoreAbstract::ObjectHeader Header = InObject.getHeader();
+	const Alembic::Abc::MetaData ObjectMetaData = InObject.getMetaData();
+	const uint32 NumChildren = InObject.getNumChildren();
+
+	if (Alembic::AbcGeom::ICurves::matches(ObjectMetaData))
+	{
+		++NumCurves;
+	}
+	else if (Alembic::AbcGeom::IPolyMesh::matches(ObjectMetaData))
+	{
+		bHasGeometry = true;
+	}
+
+	for (uint32 ChildIndex = 0; ChildIndex < NumChildren && !bHasGeometry; ++ChildIndex)
+	{
+		ValidateObject(InObject.getChild(ChildIndex), bHasGeometry, NumCurves);
+	}
+}
+
+bool FAlembicHairTranslator::CanTranslate(const FString& FilePath)
+{
+	if (!IsFileExtensionSupported(FPaths::GetExtension(FilePath)))
+	{
+		return false;
+	}
+
+	/** Factory used to generate objects*/
+	Alembic::AbcCoreFactory::IFactory Factory;
+	Alembic::AbcCoreFactory::IFactory::CoreType CompressionType = Alembic::AbcCoreFactory::IFactory::kUnknown;
+	/** Archive-typed ABC file */
+	Alembic::Abc::IArchive Archive;
+	/** Alembic typed root (top) object*/
+	Alembic::Abc::IObject TopObject;
+
+	Factory.setPolicy(Alembic::Abc::ErrorHandler::kThrowPolicy);
+	Factory.setOgawaNumStreams(12);
+
+	// Extract Archive and compression type from file
+	Archive = Factory.getArchive(TCHAR_TO_UTF8(*FilePath), CompressionType);
+	if (!Archive.valid())
+	{
+		return false;
+	}
+
+	// Get Top/root object
+	TopObject = Alembic::Abc::IObject(Archive, Alembic::Abc::kTop);
+	if (!TopObject.valid())
+	{
+		return false;
+	}
+
+	bool bHasGeometry = false;
+	int32 NumCurves = 0;
+
+	ValidateObject(TopObject, bHasGeometry, NumCurves);
+
+	return !bHasGeometry && NumCurves > 0;
+}
+
+bool FAlembicHairTranslator::IsFileExtensionSupported(const FString& FileExtension) const
+{
+	return GetSupportedFormat().StartsWith(FileExtension);
+}
+
+FString FAlembicHairTranslator::GetSupportedFormat() const
+{
+	return TEXT("abc;Alembic hair strands file");
+}
diff --git a/Engine/Plugins/Experimental/AlembicHairImporter/Source/AlembicHairTranslator/Private/AlembicHairTranslator.h b/Engine/Plugins/Experimental/AlembicHairImporter/Source/AlembicHairTranslator/Private/AlembicHairTranslator.h
new file mode 100644
--- /dev/null
+++ b/Engine/Plugins/Experimental/AlembicHairImporter/Source/AlembicHairTranslator/Private/AlembicHairTranslator.h
@@ -0,0 +1,14 @@
+// Copyright 1998-2019 Epic Games, Inc. All Rights Reserved.
+
+#pragma once
+
+#include "HairStrandsTranslator.h"
+
+class FAlembicHairTranslator : public IHairStrandsTranslator
+{
+public:
+	virtual bool Translate(const FString& FilePath, FHairDescription& OutHairDescription) override;
+	virtual bool CanTranslate(const FString& FilePath) override;
+	virtual bool IsFileExtensionSupported(const FString& FileExtension) const override;
+	virtual FString GetSupportedFormat() const override;
+};
diff --git a/Engine/Plugins/Experimental/AlembicHairImporter/Source/AlembicHairTranslator/Private/AlembicHairTranslatorModule.cpp b/Engine/Plugins/Experimental/AlembicHairImporter/Source/AlembicHairTranslator/Private/AlembicHairTranslatorModule.cpp
new file mode 100644
--- /dev/null
+++ b/Engine/Plugins/Experimental/AlembicHairImporter/Source/AlembicHairTranslator/Private/AlembicHairTranslatorModule.cpp
@@ -0,0 +1,18 @@
+// Copyright 1998-2019 Epic Games, Inc. All Rights Reserved.
+
+#include "AlembicHairTranslatorModule.h"
+
+#include "AlembicHairTranslator.h"
+#include "HairStrandsEditor.h"
+
+IMPLEMENT_MODULE(FAlembicHairTranslatorModule, AlembicHairTranslatorModule);
+
+void FAlembicHairTranslatorModule::StartupModule()
+{
+	FHairStrandsEditor::Get().RegisterHairTranslator<FAlembicHairTranslator>();
+}
+
+void FAlembicHairTranslatorModule::ShutdownModule()
+{
+}
+
diff --git a/Engine/Plugins/Experimental/AlembicImporter/Source/AlembicImporter/Classes/AlembicImportFactory.h b/Engine/Plugins/Experimental/AlembicImporter/Source/AlembicImporter/Classes/AlembicImportFactory.h
--- a/Engine/Plugins/Experimental/AlembicImporter/Source/AlembicImporter/Classes/AlembicImportFactory.h
+++ b/Engine/Plugins/Experimental/AlembicImporter/Source/AlembicImporter/Classes/AlembicImportFactory.h
@@ -37,6 +37,7 @@ class UAlembicImportFactory : public UFactory, public FReimportHandler
 	virtual FText GetDisplayName() const override;
 	virtual bool DoesSupportClass(UClass * Class) override;
 	virtual UClass* ResolveSupportedClass() override;
+	virtual bool FactoryCanImport(const FString& Filename) override;
 	virtual UObject* FactoryCreateFile(UClass* InClass, UObject* InParent, FName InName, EObjectFlags Flags, const FString& Filename, const TCHAR* Parms, FFeedbackContext* Warn, bool& bOutOperationCanceled) override;
 	//~ End UFactory Interface
 
diff --git a/Engine/Plugins/Experimental/AlembicImporter/Source/AlembicImporter/Private/AlembicImportFactory.cpp b/Engine/Plugins/Experimental/AlembicImporter/Source/AlembicImporter/Private/AlembicImportFactory.cpp
--- a/Engine/Plugins/Experimental/AlembicImporter/Source/AlembicImporter/Private/AlembicImportFactory.cpp
+++ b/Engine/Plugins/Experimental/AlembicImporter/Source/AlembicImporter/Private/AlembicImportFactory.cpp
@@ -62,6 +62,12 @@ UClass* UAlembicImportFactory::ResolveSupportedClass()
 	return UStaticMesh::StaticClass();
 }
 
+bool UAlembicImportFactory::FactoryCanImport(const FString& Filename)
+{
+	const FString Extension = FPaths::GetExtension(Filename);
+	return FPaths::GetExtension(Filename) == TEXT("abc");
+}
+
 UObject* UAlembicImportFactory::FactoryCreateFile(UClass* InClass, UObject* InParent, FName InName, EObjectFlags Flags, const FString& Filename, const TCHAR* Parms, FFeedbackContext* Warn, bool& bOutOperationCanceled)
 {
 	GEditor->GetEditorSubsystem<UImportSubsystem>()->BroadcastAssetPreImport(this, InClass, InParent, InName, TEXT("ABC"));
diff --git a/Engine/Plugins/Experimental/ChaosCloth/Source/Chaos/ChaosCloth.Build.cs b/Engine/Plugins/Experimental/ChaosCloth/Source/Chaos/ChaosCloth.Build.cs
--- a/Engine/Plugins/Experimental/ChaosCloth/Source/Chaos/ChaosCloth.Build.cs
+++ b/Engine/Plugins/Experimental/ChaosCloth/Source/Chaos/ChaosCloth.Build.cs
@@ -13,7 +13,7 @@ public class ChaosCloth : ModuleRules
             new string[] {
                 "Core",
                 "CoreUObject",
-                "ClothingSystemRuntime",
+                "ClothingSystemRuntimeNv",
                 "ClothingSystemRuntimeInterface",
                 "Engine",
                 "Chaos"
diff --git a/Engine/Plugins/Experimental/ChaosCloth/Source/Chaos/Private/ChaosCloth/ChaosClothModule.cpp b/Engine/Plugins/Experimental/ChaosCloth/Source/Chaos/Private/ChaosCloth/ChaosClothModule.cpp
--- a/Engine/Plugins/Experimental/ChaosCloth/Source/Chaos/Private/ChaosCloth/ChaosClothModule.cpp
+++ b/Engine/Plugins/Experimental/ChaosCloth/Source/Chaos/Private/ChaosCloth/ChaosClothModule.cpp
@@ -2,22 +2,37 @@
 
 #include "ChaosCloth/ChaosClothModule.h"
 #include "ChaosCloth/ChaosClothPrivate.h"
+#include "ChaosCloth/ChaosClothingSimulationFactory.h"
 #include "Modules/ModuleManager.h"
 
 //////////////////////////////////////////////////////////////////////////
 // FChaosClothModule
 
-class FChaosClothModule : public IChaosClothModuleInterface
+class FChaosClothModule : public IChaosClothModuleInterface, public IClothingSimulationFactoryClassProvider
 {
   public:
     virtual void StartupModule() override
     {
         check(GConfig);
+#if WITH_CHAOS
+		IModularFeatures::Get().RegisterModularFeature(IClothingSimulationFactoryClassProvider::FeatureName, this);
+#endif
     }
 
     virtual void ShutdownModule() override
     {
+#if WITH_CHAOS
+		IModularFeatures::Get().UnregisterModularFeature(IClothingSimulationFactoryClassProvider::FeatureName, this);
+#endif
     }
+
+	UClass* GetDefaultSimulationFactoryClass()
+	{
+#if WITH_CHAOS
+		return UChaosClothingSimulationFactory::StaticClass();
+#endif
+		return nullptr;
+	}
 };
 
 //////////////////////////////////////////////////////////////////////////
diff --git a/Engine/Plugins/Experimental/ChaosCloth/Source/Chaos/Private/ChaosCloth/ChaosClothingSimulation.cpp b/Engine/Plugins/Experimental/ChaosCloth/Source/Chaos/Private/ChaosCloth/ChaosClothingSimulation.cpp
--- a/Engine/Plugins/Experimental/ChaosCloth/Source/Chaos/Private/ChaosCloth/ChaosClothingSimulation.cpp
+++ b/Engine/Plugins/Experimental/ChaosCloth/Source/Chaos/Private/ChaosCloth/ChaosClothingSimulation.cpp
@@ -4,7 +4,8 @@
 
 #include "Async/ParallelFor.h"
 #include "Assets/ClothingAsset.h"
-#include "ClothingSimulation.h" // ClothingSystemRuntime
+#include "ClothingSimulation.h" // ClothingSystemRuntimeInterface
+#include "Utils/ClothingMeshUtils.h" // ClothingSystemRuntimeNv
 #include "Components/SkeletalMeshComponent.h"
 #include "PhysicsEngine/PhysicsAsset.h"
 #include "Chaos/Box.h"
@@ -49,7 +50,8 @@ static TAutoConsoleVariable<float> CVarClothDamping(TEXT("physics.ClothDamping")
 static TAutoConsoleVariable<float> CVarClothGravityMagnitude(TEXT("physics.ClothGravityMagnitude"), 490.f, TEXT(""));
 
 ClothingSimulation::ClothingSimulation()
-	: NumIterations(1)
+	: Asset(nullptr)
+	, NumIterations(1)
 	, EdgeStiffness(1.f)
 	, BendingStiffness(1.f)
 	, AreaStiffness(1.f)
@@ -172,8 +174,8 @@ void ClothingSimulation::CreateActor(USkeletalMeshComponent* InOwnerComponent, U
 	// It's already accounted for in our skinning matrices."), and returns all points 
 	// in that space.
 	FTransform RootBoneTransform = Context.BoneTransforms[Asset->ReferenceBoneIndex];
-	FClothingSimulationBase::SkinPhysicsMesh(
-		Asset, 
+	ClothingMeshUtils::SkinPhysicsMesh(
+		Asset->UsedBoneIndices, 
 		PhysMesh, // curr pos and norm
 		RootBoneTransform, 
 		Context.RefToLocals.GetData(), 
@@ -566,7 +568,7 @@ void ClothingSimulation::CreateActor(USkeletalMeshComponent* InOwnerComponent, U
 				"convex collision objects for physics authored as part of a LOD construct, "
 				"probably by the Apex cloth authoring system.  This is deprecated.  "
 				"Please update your asset!"),
-			*InOwnerComponent->GetOwner()->GetName(),
+			InOwnerComponent->GetOwner() ? *InOwnerComponent->GetOwner()->GetName() : TEXT("None"),
 			*InOwnerComponent->GetName(),
 			LodCollData.Spheres.Num(),
 			LodCollData.SphereConnections.Num(),
@@ -742,6 +744,9 @@ void ClothingSimulation::FillContext(USkeletalMeshComponent* InComponent, float
 
 void ClothingSimulation::Simulate(IClothingSimulationContext* InContext)
 {
+	if (!ensure(Asset))
+		return;
+
     ClothingSimulationContext* Context = static_cast<ClothingSimulationContext*>(InContext);
     if (Context->DeltaTime == 0)
         return;
@@ -753,8 +758,8 @@ void ClothingSimulation::Simulate(IClothingSimulationContext* InContext)
     const FClothPhysicalMeshData& PhysMesh = AssetLodData.PhysicalMeshData;
 
 	FTransform RootBoneTransform = Context->BoneTransforms[Asset->ReferenceBoneIndex];
-    FClothingSimulationBase::SkinPhysicsMesh(
-		Asset, 
+    ClothingMeshUtils::SkinPhysicsMesh(
+		Asset->UsedBoneIndices, 
 		PhysMesh, 
 		RootBoneTransform, 
 		Context->RefToLocals.GetData(), 
diff --git a/Engine/Plugins/Experimental/ChaosCloth/Source/Chaos/Public/ChaosCloth/ChaosClothingSimulationFactory.h b/Engine/Plugins/Experimental/ChaosCloth/Source/Chaos/Public/ChaosCloth/ChaosClothingSimulationFactory.h
--- a/Engine/Plugins/Experimental/ChaosCloth/Source/Chaos/Public/ChaosCloth/ChaosClothingSimulationFactory.h
+++ b/Engine/Plugins/Experimental/ChaosCloth/Source/Chaos/Public/ChaosCloth/ChaosClothingSimulationFactory.h
@@ -2,7 +2,7 @@
 
 #pragma once
 
-#include "ClothingSimulationFactoryInterface.h"
+#include "ClothingSimulationFactory.h"
 
 #include "ChaosClothingSimulationFactory.generated.h" // Ryan - really???
 
diff --git a/Engine/Plugins/Experimental/ChaosNiagara/Source/ChaosNiagara/Private/NiagaraDataInterfaceChaosDestruction.cpp b/Engine/Plugins/Experimental/ChaosNiagara/Source/ChaosNiagara/Private/NiagaraDataInterfaceChaosDestruction.cpp
--- a/Engine/Plugins/Experimental/ChaosNiagara/Source/ChaosNiagara/Private/NiagaraDataInterfaceChaosDestruction.cpp
+++ b/Engine/Plugins/Experimental/ChaosNiagara/Source/ChaosNiagara/Private/NiagaraDataInterfaceChaosDestruction.cpp
@@ -187,12 +187,15 @@ void UNiagaraDataInterfaceChaosDestruction::BeginDestroy()
 	Super::BeginDestroy();
 
 #if WITH_CHAOS
-	FPhysScene* Scene = GetWorld()->GetPhysicsScene();
-	if (Scene)
+	if (GetWorld())
 	{
-		Scene->UnregisterEventHandler(Chaos::EEventType::Collision, this);
-		Scene->UnregisterEventHandler(Chaos::EEventType::Breaking, this);
-		Scene->UnregisterEventHandler(Chaos::EEventType::Trailing, this);
+		FPhysScene* Scene = GetWorld()->GetPhysicsScene();
+		if (Scene)
+		{
+			Scene->UnregisterEventHandler(Chaos::EEventType::Collision, this);
+			Scene->UnregisterEventHandler(Chaos::EEventType::Breaking, this);
+			Scene->UnregisterEventHandler(Chaos::EEventType::Trailing, this);
+		}
 	}
 #endif
 }
diff --git a/Engine/Plugins/Experimental/GeometryCollectionPlugin/GeometryCollectionPlugin.uplugin b/Engine/Plugins/Experimental/GeometryCollectionPlugin/GeometryCollectionPlugin.uplugin
--- a/Engine/Plugins/Experimental/GeometryCollectionPlugin/GeometryCollectionPlugin.uplugin
+++ b/Engine/Plugins/Experimental/GeometryCollectionPlugin/GeometryCollectionPlugin.uplugin
@@ -18,7 +18,8 @@
 	{
       "Name": "GeometryCollectionExampleCore",
       "Type": "Runtime",
-      "LoadingPhase": "Default"
+	  "LoadingPhase": "Default",
+      "WhitelistPlatforms" : [ "Win32", "Win64" ]
     },
     {
       "Name": "GeometryCollectionEditor",
diff --git a/Engine/Plugins/Experimental/GeometryCollectionPlugin/Source/GeometryCollectionExampleCore/GeometryCollectionExampleCore.Build.cs b/Engine/Plugins/Experimental/GeometryCollectionPlugin/Source/GeometryCollectionExampleCore/GeometryCollectionExampleCore.Build.cs
--- a/Engine/Plugins/Experimental/GeometryCollectionPlugin/Source/GeometryCollectionExampleCore/GeometryCollectionExampleCore.Build.cs
+++ b/Engine/Plugins/Experimental/GeometryCollectionPlugin/Source/GeometryCollectionExampleCore/GeometryCollectionExampleCore.Build.cs
@@ -9,19 +9,41 @@ namespace UnrealBuildTool.Rules
 			PrivateIncludePaths.Add("GeometryCollectionExampleCore/Private");
             PublicIncludePaths.Add(ModuleDirectory + "/Public");
 
-			PublicDependencyModuleNames.AddRange(
+            PublicDependencyModuleNames.AddRange(
 				new string[]
 				{
-					"Core",
+                    "Core",
 					"CoreUObject",
                     "GeometryCollectionCore",
                     "GeometryCollectionSimulationCore",
-					"ChaosSolvers",
+                    "GoogleTest",
+                    "ChaosSolvers",
                     "Chaos",
                     "FieldSystemCore",
                     "FieldSystemSimulationCore"
                 }
-                );
-		}
+            );
+
+			if (Target.Platform == UnrealTargetPlatform.Win32 || Target.Platform == UnrealTargetPlatform.Win64 || Target.Platform == UnrealTargetPlatform.XboxOne)
+			{
+				PublicDefinitions.Add("GTEST_OS_WINDOWS=1");
+			}
+			else if (Target.Platform == UnrealTargetPlatform.Mac)
+			{
+				PublicDefinitions.Add("GTEST_OS_MAC=1");
+			}
+			else if (Target.Platform == UnrealTargetPlatform.IOS || Target.Platform == UnrealTargetPlatform.TVOS)
+			{
+				PublicDefinitions.Add("GTEST_OS_IOS=1");
+			}
+			else if (Target.Platform == UnrealTargetPlatform.Android || Target.Platform == UnrealTargetPlatform.Quail || Target.Platform == UnrealTargetPlatform.Lumin)
+			{
+				PublicDefinitions.Add("GTEST_OS_LINUX_ANDROID=1");
+			}
+			else if (Target.Platform == UnrealTargetPlatform.Linux || Target.Platform == UnrealTargetPlatform.PS4)
+			{
+				PublicDefinitions.Add("GTEST_OS_LINUX=1");
+			}
+        }
 	}
 }
diff --git a/Engine/Plugins/Experimental/GeometryCollectionPlugin/Source/GeometryCollectionExampleCore/Private/GeometryCollection/GeometryCollectionExample.cpp b/Engine/Plugins/Experimental/GeometryCollectionPlugin/Source/GeometryCollectionExampleCore/Private/GeometryCollection/GeometryCollectionExample.cpp
--- a/Engine/Plugins/Experimental/GeometryCollectionPlugin/Source/GeometryCollectionExampleCore/Private/GeometryCollection/GeometryCollectionExample.cpp
+++ b/Engine/Plugins/Experimental/GeometryCollectionPlugin/Source/GeometryCollectionExampleCore/Private/GeometryCollection/GeometryCollectionExample.cpp
@@ -18,7 +18,7 @@
 
 namespace GeometryCollectionExample
 {
-#define RUN_EXAMPLE(X) X<float>(RESPONSE());
+#define RUN_EXAMPLE(X) X<float>();
 
 #if TODO_REIMPLEMENT_RIGID_CLUSTERING
 	template<class RESPONSE>
diff --git a/Engine/Plugins/Experimental/GeometryCollectionPlugin/Source/GeometryCollectionExampleCore/Private/GeometryCollection/GeometryCollectionExampleBoneHierarchy.cpp b/Engine/Plugins/Experimental/GeometryCollectionPlugin/Source/GeometryCollectionExampleCore/Private/GeometryCollection/GeometryCollectionExampleBoneHierarchy.cpp
--- a/Engine/Plugins/Experimental/GeometryCollectionPlugin/Source/GeometryCollectionExampleCore/Private/GeometryCollection/GeometryCollectionExampleBoneHierarchy.cpp
+++ b/Engine/Plugins/Experimental/GeometryCollectionPlugin/Source/GeometryCollectionExampleCore/Private/GeometryCollection/GeometryCollectionExampleBoneHierarchy.cpp
@@ -11,19 +11,16 @@ namespace GeometryCollectionExample
 {
 #if INCLUDE_CHAOS
 	template <class TImplicitShape>
-	bool AllOnSurface(const TImplicitShape *Shape, const TArray<Chaos::TVector<float, 3>> &Points, const float Tolerance = KINDA_SMALL_NUMBER)
+	void AllOnSurface(const TImplicitShape *Shape, const TArray<Chaos::TVector<float, 3>> &Points, const float Tolerance = KINDA_SMALL_NUMBER)
 	{
-		bool RetVal = true;
 		for (auto &Pt : Points)
 		{
 			const float Phi = Shape->SignedDistance(Pt);
-			RetVal &= FMath::Abs(Phi) <= Tolerance;
-			check(RetVal);
+			EXPECT_NEAR(Phi, 0, Tolerance);
 		}
-		return RetVal;
 	}
 
-	void TestSphere(ExampleResponse& R, const Chaos::TVector<float, 3>& Center, const float Radius, const FTransform &BoneRelXf)
+	void TestSphere(const Chaos::TVector<float, 3>& Center, const float Radius, const FTransform &BoneRelXf)
 	{
 		Chaos::TSphere<float, 3>* Sphere = new Chaos::TSphere<float, 3>(Center, Radius);
 		TUniquePtr<FAnalyticImplicitGroup> Group1(new FAnalyticImplicitGroup("Root", 0));
@@ -34,10 +31,9 @@ namespace GeometryCollectionExample
 		check(Points);
 		check(Points->Num());
 		Chaos::TImplicitObject<float, 3> *Implicit = Group1->BuildSimImplicitObject();
-		R.ExpectTrue(AllOnSurface(Implicit, *Points));
+		AllOnSurface(Implicit, *Points);
 	}
 	void TestSphere2(
-		ExampleResponse& R, 
 		const Chaos::TVector<float, 3>& Center1, 
 		const Chaos::TVector<float, 3>& Center2, 
 		const float Radius1, 
@@ -56,15 +52,15 @@ namespace GeometryCollectionExample
 		check(Points);
 		check(Points->Num());
 		Chaos::TImplicitObject<float, 3> *Implicit = Group1->BuildSimImplicitObject();
-		R.ExpectTrue(AllOnSurface(Implicit, *Points));
+		AllOnSurface(Implicit, *Points);
 	}
 
-	void RunAnalyticImplicitGroupTest(ExampleResponse& R)
+	void RunAnalyticImplicitGroupTest()
 	{
-		TestSphere(R, Chaos::TVector<float, 3>(0, 0, 0), 1.f, FTransform::Identity);
-		TestSphere(R, Chaos::TVector<float, 3>(0, 0, 0), 1.f, FTransform(FVector(1,0,0)));
-		TestSphere(R, Chaos::TVector<float, 3>(0, 0, 0), 1.f, FTransform(FVector(1,1,0)));
-		TestSphere(R, Chaos::TVector<float, 3>(0, 0, 0), 1.f, FTransform(FVector(1,1,1)));
+		TestSphere(Chaos::TVector<float, 3>(0, 0, 0), 1.f, FTransform::Identity);
+		TestSphere(Chaos::TVector<float, 3>(0, 0, 0), 1.f, FTransform(FVector(1,0,0)));
+		TestSphere(Chaos::TVector<float, 3>(0, 0, 0), 1.f, FTransform(FVector(1,1,0)));
+		TestSphere(Chaos::TVector<float, 3>(0, 0, 0), 1.f, FTransform(FVector(1,1,1)));
 
 		//This test fails in Dev-Physics for spheres of larger radii
 		//TestSphere(R, Chaos::TVector<float, 3>(0, 0, 0), 10.f, FTransform::Identity);
@@ -72,7 +68,7 @@ namespace GeometryCollectionExample
 		//TestSphere(R, Chaos::TVector<float, 3>(0, 0, 0), 10.f, FTransform(FVector(1,1,0)));
 		//TestSphere(R, Chaos::TVector<float, 3>(0, 0, 0), 10.f, FTransform(FVector(1,1,1)));
 
-		TestSphere2(R, 
+		TestSphere2( 
 			Chaos::TVector<float, 3>(0, 0, 0), 
 			Chaos::TVector<float, 3>(0, 0, 0), 
 			1.f, 
@@ -81,7 +77,7 @@ namespace GeometryCollectionExample
 			FTransform(FVector(2,0,0)));
 	}
 
-	void RunBoneHierarchyTest(ExampleResponse& R)
+	void RunBoneHierarchyTest()
 	{
 		FBoneHierarchy Hierarchy;
 		Hierarchy.InitPreAdd(2);
@@ -144,13 +140,12 @@ namespace GeometryCollectionExample
 #endif
 
 	template<class T>
-	bool TestImplicitBoneHierarchy(ExampleResponse&& R)
+	void TestImplicitBoneHierarchy()
 	{
 #if INCLUDE_CHAOS
-		RunAnalyticImplicitGroupTest(R);
-		RunBoneHierarchyTest(R);
+		RunAnalyticImplicitGroupTest();
+		RunBoneHierarchyTest();
 #endif
-		return !R.HasError();
 	}
-	template bool TestImplicitBoneHierarchy<float>(ExampleResponse&& R);
+	template void TestImplicitBoneHierarchy<float>();
 } // namespace GeometryCollectionExample
diff --git a/Engine/Plugins/Experimental/GeometryCollectionPlugin/Source/GeometryCollectionExampleCore/Private/GeometryCollection/GeometryCollectionExampleClean.cpp b/Engine/Plugins/Experimental/GeometryCollectionPlugin/Source/GeometryCollectionExampleCore/Private/GeometryCollection/GeometryCollectionExampleClean.cpp
--- a/Engine/Plugins/Experimental/GeometryCollectionPlugin/Source/GeometryCollectionExampleCore/Private/GeometryCollection/GeometryCollectionExampleClean.cpp
+++ b/Engine/Plugins/Experimental/GeometryCollectionPlugin/Source/GeometryCollectionExampleCore/Private/GeometryCollection/GeometryCollectionExampleClean.cpp
@@ -14,7 +14,7 @@
 namespace GeometryCollectionExample
 {
 	template<class T>
-	bool TestDeleteCoincidentVertices(ExampleResponse&& R)
+	void TestDeleteCoincidentVertices()
 	{
 		TSharedPtr<FGeometryCollection> Collection = GeometryCollection::MakeCubeElement(FTransform(FQuat::MakeFromEuler(FVector(0.f, 0.f, 0.f)), FVector(0.f, 0.f, 0.f)), FVector(1.0));
 		Collection->AppendGeometry(*GeometryCollection::MakeCubeElement(FTransform(FQuat::MakeFromEuler(FVector(0.f, 0.f, 0.f)), FVector(1.f, 0.f, 0.f)), FVector(1.0)));
@@ -35,22 +35,20 @@ namespace GeometryCollectionExample
 		GeometryCollectionAlgo::DeleteCoincidentVertices(Coll, 1e-2);
 
 		
-		R.ExpectTrue(Coll->NumElements(FGeometryCollection::VerticesGroup) == 24);
-		R.ExpectTrue(Coll->NumElements(FGeometryCollection::FacesGroup) == 36);
+		EXPECT_EQ(Coll->NumElements(FGeometryCollection::VerticesGroup), 24);
+		EXPECT_EQ(Coll->NumElements(FGeometryCollection::FacesGroup), 36);
 
 		GeometryCollectionAlgo::DeleteZeroAreaFaces(Coll, 1e-4);
 
-		R.ExpectTrue(Coll->NumElements(FGeometryCollection::VerticesGroup) == 24);
-		R.ExpectTrue(Coll->NumElements(FGeometryCollection::FacesGroup) == 36);
-
-		return !R.HasError();
+		EXPECT_EQ(Coll->NumElements(FGeometryCollection::VerticesGroup), 24);
+		EXPECT_EQ(Coll->NumElements(FGeometryCollection::FacesGroup), 36);
 	}
-	template bool TestDeleteCoincidentVertices<float>(ExampleResponse&& R);
+	template void TestDeleteCoincidentVertices<float>();
 
 	
 
 	template<class T>
-	bool TestDeleteCoincidentVertices2(ExampleResponse&& R)
+	void TestDeleteCoincidentVertices2()
 	{
 		FGeometryCollection* Coll = FGeometryCollection::NewGeometryCollection(FracturedGeometry::RawVertexArray,
 																   			   FracturedGeometry::RawIndicesArray,
@@ -62,20 +60,18 @@ namespace GeometryCollectionExample
 																			   FracturedGeometry::RawSimulationTypeArray,
 																			   FracturedGeometry::RawStatusFlagsArray);
 
-		R.ExpectTrue(Coll->NumElements(FGeometryCollection::VerticesGroup) == 667);
-		R.ExpectTrue(Coll->NumElements(FGeometryCollection::FacesGroup) == 493);
+		EXPECT_EQ(Coll->NumElements(FGeometryCollection::VerticesGroup), 667);
+		EXPECT_EQ(Coll->NumElements(FGeometryCollection::FacesGroup), 493);
 
 		GeometryCollectionAlgo::DeleteCoincidentVertices(Coll, 1e-2);
 
-		R.ExpectTrue(Coll->NumElements(FGeometryCollection::VerticesGroup) == 270);
-		R.ExpectTrue(Coll->NumElements(FGeometryCollection::FacesGroup) == 493);
-
-		return !R.HasError();
+		EXPECT_EQ(Coll->NumElements(FGeometryCollection::VerticesGroup), 270);
+		EXPECT_EQ(Coll->NumElements(FGeometryCollection::FacesGroup), 493);
 	}
-	template bool TestDeleteCoincidentVertices2<float>(ExampleResponse&& R);
+	template void TestDeleteCoincidentVertices2<float>();
 
 	template<class T>
-	bool TestDeleteZeroAreaFaces(ExampleResponse&& R)
+	void TestDeleteZeroAreaFaces()
 	{
 		FGeometryCollection* Coll = FGeometryCollection::NewGeometryCollection(FracturedGeometry::RawVertexArray,
 																			   FracturedGeometry::RawIndicesArray,
@@ -88,20 +84,18 @@ namespace GeometryCollectionExample
 																			   FracturedGeometry::RawStatusFlagsArray
 			);
 
-		R.ExpectTrue(Coll->NumElements(FGeometryCollection::VerticesGroup) == 667);
-		R.ExpectTrue(Coll->NumElements(FGeometryCollection::FacesGroup) == 493);
+		EXPECT_EQ(Coll->NumElements(FGeometryCollection::VerticesGroup), 667);
+		EXPECT_EQ(Coll->NumElements(FGeometryCollection::FacesGroup), 493);
 
 		GeometryCollectionAlgo::DeleteZeroAreaFaces(Coll, 1e-4);
 
-		R.ExpectTrue(Coll->NumElements(FGeometryCollection::VerticesGroup) == 667);
-		R.ExpectTrue(Coll->NumElements(FGeometryCollection::FacesGroup) == 493);
-
-		return !R.HasError();
+		EXPECT_EQ(Coll->NumElements(FGeometryCollection::VerticesGroup), 667);
+		EXPECT_EQ(Coll->NumElements(FGeometryCollection::FacesGroup), 493);
 	}
-	template bool TestDeleteZeroAreaFaces<float>(ExampleResponse&& R);
+	template void TestDeleteZeroAreaFaces<float>();
 
 	template<class T>
-	bool TestFillHoles(ExampleResponse&& R)
+	void TestFillHoles()
 	{
 		FGeometryCollection* Coll = FGeometryCollection::NewGeometryCollection(FracturedGeometry::RawVertexArray,
 			FracturedGeometry::RawIndicesArray,
@@ -112,8 +106,8 @@ namespace GeometryCollectionExample
 			FracturedGeometry::RawChildrenArray,
 			FracturedGeometry::RawSimulationTypeArray,
 			FracturedGeometry::RawStatusFlagsArray);
-		R.ExpectTrue(Coll->NumElements(FGeometryCollection::VerticesGroup) == 667);
-		R.ExpectTrue(Coll->NumElements(FGeometryCollection::FacesGroup) == 493);
+		EXPECT_EQ(Coll->NumElements(FGeometryCollection::VerticesGroup), 667);
+		EXPECT_EQ(Coll->NumElements(FGeometryCollection::FacesGroup), 493);
 
 		TArray<TArray<TArray<int32>>> BoundaryVertexIndices;
 		Coll->RemoveElements(FGeometryCollection::FacesGroup, { 0,1,2 });
@@ -147,38 +141,36 @@ namespace GeometryCollectionExample
 		};
 		
 		int32 TinyFacesBefore = CountTinyFaces(Coll);
-		R.ExpectTrue(CountHoles(BoundaryVertexIndices) == 3);
+		EXPECT_EQ(CountHoles(BoundaryVertexIndices), 3);
 		GeometryCollectionAlgo::TriangulateBoundaries(Coll, BoundaryVertexIndices);
 		int32 TinyFacesAfter = CountTinyFaces(Coll);
-		R.ExpectTrue(CountTinyFaces(Coll) == TinyFacesBefore);
+		EXPECT_EQ(CountTinyFaces(Coll), TinyFacesBefore);
 
 		BoundaryVertexIndices.Empty();
 		GeometryCollectionAlgo::FindOpenBoundaries(Coll, 1e-2, BoundaryVertexIndices);
-		R.ExpectTrue(CountHoles(BoundaryVertexIndices) == 2);
+		EXPECT_EQ(CountHoles(BoundaryVertexIndices), 2);
 
 		GeometryCollectionAlgo::TriangulateBoundaries(Coll, BoundaryVertexIndices, true, 0);
 		BoundaryVertexIndices.Empty();
 		GeometryCollectionAlgo::FindOpenBoundaries(Coll, 1e-2, BoundaryVertexIndices);
-		R.ExpectTrue(CountHoles(BoundaryVertexIndices) == 0);
-		R.ExpectTrue(CountTinyFaces(Coll) > TinyFacesBefore);
+		EXPECT_EQ(CountHoles(BoundaryVertexIndices), 0);
+		EXPECT_GT(CountTinyFaces(Coll), TinyFacesBefore);
 
 
-		R.ExpectTrue(Coll->NumElements(FGeometryCollection::VerticesGroup) == 667);
-		R.ExpectTrue(Coll->NumElements(FGeometryCollection::FacesGroup) == 496);
+		EXPECT_EQ(Coll->NumElements(FGeometryCollection::VerticesGroup), 667);
+		EXPECT_EQ(Coll->NumElements(FGeometryCollection::FacesGroup), 496);
 
-		R.ExpectTrue(Coll->NumElements(FGeometryCollection::VerticesGroup) == 667);
-		R.ExpectTrue(Coll->NumElements(FGeometryCollection::FacesGroup) == 496);
+		EXPECT_EQ(Coll->NumElements(FGeometryCollection::VerticesGroup), 667);
+		EXPECT_EQ(Coll->NumElements(FGeometryCollection::FacesGroup), 496);
 
-		R.ExpectTrue(Coll->HasContiguousFaces());
-		R.ExpectTrue(Coll->HasContiguousVertices());
-		R.ExpectTrue(GeometryCollectionAlgo::HasValidGeometryReferences(Coll));
-
-		return !R.HasError();
+		EXPECT_TRUE(Coll->HasContiguousFaces());
+		EXPECT_TRUE(Coll->HasContiguousVertices());
+		EXPECT_TRUE(GeometryCollectionAlgo::HasValidGeometryReferences(Coll));
 	}
-	template bool TestFillHoles<float>(ExampleResponse&& R);
+	template void TestFillHoles<float>();
 
 	template<class T>
-	bool TestDeleteHiddenFaces(ExampleResponse&& R)
+	void TestDeleteHiddenFaces()
 	{
 		FGeometryCollection* Coll = FGeometryCollection::NewGeometryCollection(FracturedGeometry::RawVertexArray,
 																			   FracturedGeometry::RawIndicesArray,
@@ -190,8 +182,8 @@ namespace GeometryCollectionExample
 																			   FracturedGeometry::RawSimulationTypeArray,
 																			   FracturedGeometry::RawStatusFlagsArray);
 
-		R.ExpectTrue(Coll->NumElements(FGeometryCollection::VerticesGroup) == 667);
-		R.ExpectTrue(Coll->NumElements(FGeometryCollection::FacesGroup) == 493);
+		EXPECT_EQ(Coll->NumElements(FGeometryCollection::VerticesGroup), 667);
+		EXPECT_EQ(Coll->NumElements(FGeometryCollection::FacesGroup), 493);
 
 		TManagedArray<bool>& VisibleArray = Coll->Visible;
 
@@ -204,16 +196,14 @@ namespace GeometryCollectionExample
 			}
 		}
 
-		R.ExpectTrue(Coll->NumElements(FGeometryCollection::VerticesGroup) == 667);
-		R.ExpectTrue(Coll->NumElements(FGeometryCollection::FacesGroup) == 493);
+		EXPECT_EQ(Coll->NumElements(FGeometryCollection::VerticesGroup), 667);
+		EXPECT_EQ(Coll->NumElements(FGeometryCollection::FacesGroup), 493);
 
 		GeometryCollectionAlgo::DeleteHiddenFaces(Coll);
 
-		R.ExpectTrue( Coll->NumElements(FGeometryCollection::VerticesGroup) == 667);
-		R.ExpectTrue(Coll->NumElements(FGeometryCollection::FacesGroup) == 394);
-
-		return !R.HasError();
+		EXPECT_EQ(Coll->NumElements(FGeometryCollection::VerticesGroup), 667);
+		EXPECT_EQ(Coll->NumElements(FGeometryCollection::FacesGroup), 394);
 	}
-	template bool TestDeleteHiddenFaces<float>(ExampleResponse&& R);
+	template void TestDeleteHiddenFaces<float>();
 }
 
