Updated hair import to Alembic hair schema 1.2:
* No groom_width attribute, using getWidthsParam instead
* Groom-scoped attributes as UserProperties
* Using GeometryScope to determine attribute scope
* Profiling tags

(ushell-p4-cherrypick of 9654805 by [REDACTED_PERSON])

diff --git a/Engine/Plugins/Experimental/AlembicHairImporter/Source/AlembicHairTranslator/Private/AlembicHairTranslator.cpp b/Engine/Plugins/Experimental/AlembicHairImporter/Source/AlembicHairTranslator/Private/AlembicHairTranslator.cpp
--- a/Engine/Plugins/Experimental/AlembicHairImporter/Source/AlembicHairTranslator/Private/AlembicHairTranslator.cpp
+++ b/Engine/Plugins/Experimental/AlembicHairImporter/Source/AlembicHairTranslator/Private/AlembicHairTranslator.cpp
@@ -40,6 +40,28 @@ namespace AlembicHairTranslatorUtils
 		return AttributeName.StartsWith(TEXT("groom_"));
 	}
 
+	// Groom attributes from UserProperties (as ScalarProperty)
+	template <typename AbcParamType, typename AttributeType>
+	void SetGroomAttributes(FHairDescription& HairDescription, const Alembic::AbcGeom::ICompoundProperty& Parameters, const std::string& PropName)
+	{
+		FName AttributeName(ANSI_TO_TCHAR(PropName.c_str()));
+
+		AbcParamType Param(Parameters, PropName);
+		AttributeType ParamValue = Param.getValue();
+		SetGroomAttribute(HairDescription, FGroomID(0), AttributeName, ParamValue);
+	}
+
+	template <>
+	void SetGroomAttributes<Alembic::Abc::IStringProperty, FName>(FHairDescription& HairDescription, const Alembic::AbcGeom::ICompoundProperty& Parameters, const std::string& PropName)
+	{
+		FName AttributeName(ANSI_TO_TCHAR(PropName.c_str()));
+
+		Alembic::Abc::IStringProperty Param(Parameters, PropName);
+		FName ParamValue = FName(Param.getValue().c_str());
+		SetGroomAttribute(HairDescription, FGroomID(0), AttributeName, ParamValue);
+	}
+
+	// Groom attributes from arbitrary GeomParams (as ArrayProperty with only one value)
 	template <typename AbcParamType, typename AbcArraySampleType, typename AttributeType>
 	void SetGroomAttributes(FHairDescription& HairDescription, const Alembic::AbcGeom::ICompoundProperty& Parameters, const std::string& PropName)
 	{
@@ -109,17 +131,38 @@ namespace AlembicHairTranslatorUtils
 				{
 				case Alembic::Util::kInt16POD:
 				{
-					SetGroomAttributes<Alembic::AbcGeom::IInt16GeomParam, Alembic::Abc::Int16ArraySamplePtr, int>(HairDescription, Parameters, PropName);
+					if (PropType == Alembic::Abc::kScalarProperty)
+					{
+						SetGroomAttributes<Alembic::Abc::IInt16Property, int>(HairDescription, Parameters, PropName);
+					}
+					else
+					{
+						SetGroomAttributes<Alembic::AbcGeom::IInt16GeomParam, Alembic::Abc::Int16ArraySamplePtr, int>(HairDescription, Parameters, PropName);
+					}
 				}
 				break;
 				case Alembic::Util::kInt32POD:
 				{
-					SetGroomAttributes<Alembic::AbcGeom::IInt32GeomParam, Alembic::Abc::Int32ArraySamplePtr, int>(HairDescription, Parameters, PropName);
+					if (PropType == Alembic::Abc::kScalarProperty)
+					{
+						SetGroomAttributes<Alembic::Abc::IInt32Property, int>(HairDescription, Parameters, PropName);
+					}
+					else
+					{
+						SetGroomAttributes<Alembic::AbcGeom::IInt32GeomParam, Alembic::Abc::Int32ArraySamplePtr, int>(HairDescription, Parameters, PropName);
+					}
 				}
 				break;
 				case Alembic::Util::kStringPOD:
 				{
-					SetGroomAttributes<Alembic::AbcGeom::IStringGeomParam, Alembic::Abc::StringArraySamplePtr, FName>(HairDescription, Parameters, PropName);
+					if (PropType == Alembic::Abc::kScalarProperty)
+					{
+						SetGroomAttributes<Alembic::Abc::IStringProperty, FName>(HairDescription, Parameters, PropName);
+					}
+					else
+					{
+						SetGroomAttributes<Alembic::AbcGeom::IStringGeomParam, Alembic::Abc::StringArraySamplePtr, FName>(HairDescription, Parameters, PropName);
+					}
 				}
 				break;
 				case Alembic::Util::kFloat32POD:
@@ -160,20 +203,33 @@ namespace AlembicHairTranslatorUtils
 	}
 
 	template <typename AbcParamType, typename AbcArraySampleType, typename AttributeType>
-	void ConvertAlembicAttribute(FHairDescription& HairDescription, FStrandID StrandID, int32 StartVertexID, int32 NumVertices, const Alembic::AbcGeom::ICompoundProperty& Parameters, const std::string& PropName)
+	void ConvertAlembicAttribute(FHairDescription& HairDescription, int32 StartStrandID, int32 NumStrands, int32 StartVertexID, int32 NumVertices, const Alembic::AbcGeom::ICompoundProperty& Parameters, const std::string& PropName)
 	{
 		FName AttributeName(ANSI_TO_TCHAR(PropName.c_str()));
 
-		// The number of values in Param determines the scope on which to set the hair attribute
+		// Get the param scope and check if it's supported, otherwise fall back to using the number of values in Param to deduce its scope
 		AbcParamType Param(Parameters, PropName);
 		AbcArraySampleType ParamValues = Param.getExpandedValue().getVals();
-		if (ParamValues->size() == 1)
+		Alembic::AbcGeom::GeometryScope Scope = Param.getScope();
+		int32 NumValues = ParamValues->size();
+
+		// Check the supported scope: UniformScope or VertexScope (1.2)
+		// ConstantScope (1.0) not useful so use NumStrands or NumVertices to determine scope
+		if (Scope == Alembic::AbcGeom::kUniformScope || NumValues == NumStrands)
 		{
-			AttributeType ParamValue = (*ParamValues)[0];
-			SetHairStrandAttribute(HairDescription, StrandID, AttributeName, ParamValue);
+			TStrandAttributesRef<AttributeType> StrandAttributeRef = HairDescription.StrandAttributes().GetAttributesRef<AttributeType>(AttributeName);
+			if (!StrandAttributeRef.IsValid())
+			{
+				HairDescription.StrandAttributes().RegisterAttribute<AttributeType>(AttributeName);
+				StrandAttributeRef = HairDescription.StrandAttributes().GetAttributesRef<AttributeType>(AttributeName);
+			}
 
+			for (int32 StrandIndex = 0; StrandIndex < NumValues; ++StrandIndex)
+			{
+				StrandAttributeRef[FStrandID(StartStrandID + StrandIndex)] = (*ParamValues)[StrandIndex];
+			}
 		}
-		else if (ParamValues->size() == NumVertices)
+		else if (Scope == Alembic::AbcGeom::kVertexScope || NumValues == NumVertices)
 		{
 			TVertexAttributesRef<AttributeType> VertexAttributeRef = HairDescription.VertexAttributes().GetAttributesRef<AttributeType>(AttributeName);
 			if (!VertexAttributeRef.IsValid())
@@ -190,23 +246,39 @@ namespace AlembicHairTranslatorUtils
 	}
 
 	template <typename AbcParamType, typename AbcArraySampleType, typename AttributeType>
-	void ConvertAlembicAttribute(FHairDescription& HairDescription, FStrandID StrandID, int32 StartVertexID, int32 NumVertices, const Alembic::AbcGeom::ICompoundProperty& Parameters, const std::string& PropName, uint8 Extent)
+	void ConvertAlembicAttribute(FHairDescription& HairDescription, int32 StartStrandID, int32 NumStrands, int32 StartVertexID, int32 NumVertices, const Alembic::AbcGeom::ICompoundProperty& Parameters, const std::string& PropName, uint8 Extent)
 	{
 		FName AttributeName(ANSI_TO_TCHAR(PropName.c_str()));
 
-		// The number of values in Param determines the scope on which to set the hair attribute
+		// Get the param scope and check if it's supported, otherwise fall back to using the number of values in Param to deduce its scope
 		AbcParamType Param(Parameters, PropName);
 		AbcArraySampleType ParamValues = Param.getExpandedValue().getVals();
-		if (ParamValues->size() == 1)
+		Alembic::AbcGeom::GeometryScope Scope = Param.getScope();
+		int32 NumValues = ParamValues->size();
+
+		// Check the supported scope: UniformScope or VertexScope (1.2)
+		// ConstantScope (1.0) not useful so use NumStrands or NumVertices to determine scope
+		if (Scope == Alembic::AbcGeom::kUniformScope || NumValues == NumStrands)
 		{
-			AttributeType ParamValue;
-			for (int32 Index = 0; Index < Extent; ++Index)
+			TStrandAttributesRef<AttributeType> StrandAttributeRef = HairDescription.StrandAttributes().GetAttributesRef<AttributeType>(AttributeName);
+			if (!StrandAttributeRef.IsValid())
 			{
-				ParamValue[Index] = (*ParamValues)[0][Index];
+				HairDescription.StrandAttributes().RegisterAttribute<AttributeType>(AttributeName);
+				StrandAttributeRef = HairDescription.StrandAttributes().GetAttributesRef<AttributeType>(AttributeName);
+			}
+
+			for (int32 StrandIndex = 0; StrandIndex < NumValues; ++StrandIndex)
+			{
+				AttributeType ParamValue;
+				for (int32 Index = 0; Index < Extent; ++Index)
+				{
+					ParamValue[Index] = (*ParamValues)[StrandIndex][Index];
+				}
+
+				StrandAttributeRef[FStrandID(StartStrandID + StrandIndex)] = ParamValue;
 			}
-			SetHairStrandAttribute(HairDescription, StrandID, AttributeName, ParamValue);
 		}
-		else if (ParamValues->size() == NumVertices)
+		else if (Scope == Alembic::AbcGeom::kVertexScope || NumValues == NumVertices)
 		{
 			TVertexAttributesRef<AttributeType> VertexAttributeRef = HairDescription.VertexAttributes().GetAttributesRef<AttributeType>(AttributeName);
 			if (!VertexAttributeRef.IsValid())
@@ -229,8 +301,10 @@ namespace AlembicHairTranslatorUtils
 	}
 
 	/** Convert the given Alembic parameters to hair attributes in the proper scope */
-	void ConvertAlembicAttributes(FHairDescription& HairDescription, FStrandID StrandID, int32 StartVertexID, int32 NumVertices, const Alembic::AbcGeom::ICompoundProperty& Parameters)
+	void ConvertAlembicAttributes(FHairDescription& HairDescription, int32 StartStrandID, int32 NumStrands, int32 StartVertexID, int32 NumVertices, const Alembic::AbcGeom::ICompoundProperty& Parameters)
 	{
+		TRACE_CPUPROFILER_EVENT_SCOPE(ConvertAlembicAttributes);
+
 		for (int Index = 0; Index < Parameters.getNumProperties(); ++Index)
 		{
 			Alembic::Abc::PropertyHeader PropertyHeader = Parameters.getPropertyHeader(Index);
@@ -252,22 +326,22 @@ namespace AlembicHairTranslatorUtils
 				{
 				case Alembic::Util::kBooleanPOD:
 				{
-					ConvertAlembicAttribute<Alembic::AbcGeom::IBoolGeomParam, Alembic::Abc::BoolArraySamplePtr, bool>(HairDescription, StrandID, StartVertexID, NumVertices, Parameters, PropName);
+					ConvertAlembicAttribute<Alembic::AbcGeom::IBoolGeomParam, Alembic::Abc::BoolArraySamplePtr, bool>(HairDescription, StartStrandID, NumStrands, StartVertexID, NumVertices, Parameters, PropName);
 				}
 				break;
 				case Alembic::Util::kInt8POD:
 				{
-					ConvertAlembicAttribute<Alembic::AbcGeom::ICharGeomParam, Alembic::Abc::CharArraySamplePtr, int>(HairDescription, StrandID, StartVertexID, NumVertices, Parameters, PropName);
+					ConvertAlembicAttribute<Alembic::AbcGeom::ICharGeomParam, Alembic::Abc::CharArraySamplePtr, int>(HairDescription, StartStrandID, NumStrands, StartVertexID, NumVertices, Parameters, PropName);
 				}
 				break;
 				case Alembic::Util::kInt16POD:
 				{
-					ConvertAlembicAttribute<Alembic::AbcGeom::IInt16GeomParam, Alembic::Abc::Int16ArraySamplePtr, int>(HairDescription, StrandID, StartVertexID, NumVertices, Parameters, PropName);
+					ConvertAlembicAttribute<Alembic::AbcGeom::IInt16GeomParam, Alembic::Abc::Int16ArraySamplePtr, int>(HairDescription, StartStrandID, NumStrands, StartVertexID, NumVertices, Parameters, PropName);
 				}
 				break;
 				case Alembic::Util::kInt32POD:
 				{
-					ConvertAlembicAttribute<Alembic::AbcGeom::IInt32GeomParam, Alembic::Abc::Int32ArraySamplePtr, int>(HairDescription, StrandID, StartVertexID, NumVertices, Parameters, PropName);
+					ConvertAlembicAttribute<Alembic::AbcGeom::IInt32GeomParam, Alembic::Abc::Int32ArraySamplePtr, int>(HairDescription, StartStrandID, NumStrands, StartVertexID, NumVertices, Parameters, PropName);
 				}
 				break;
 				case Alembic::Util::kFloat32POD:
@@ -275,13 +349,13 @@ namespace AlembicHairTranslatorUtils
 					switch (Extent)
 					{
 					case 1:
-						ConvertAlembicAttribute<Alembic::AbcGeom::IFloatGeomParam, Alembic::Abc::FloatArraySamplePtr, float>(HairDescription, StrandID, StartVertexID, NumVertices, Parameters, PropName);
+						ConvertAlembicAttribute<Alembic::AbcGeom::IFloatGeomParam, Alembic::Abc::FloatArraySamplePtr, float>(HairDescription, StartStrandID, NumStrands, StartVertexID, NumVertices, Parameters, PropName);
 						break;
 					case 2:
-						ConvertAlembicAttribute<Alembic::AbcGeom::IV2fGeomParam, Alembic::Abc::V2fArraySamplePtr, FVector2D>(HairDescription, StrandID, StartVertexID, NumVertices, Parameters, PropName, DataType.getExtent());
+						ConvertAlembicAttribute<Alembic::AbcGeom::IV2fGeomParam, Alembic::Abc::V2fArraySamplePtr, FVector2D>(HairDescription, StartStrandID, NumStrands, StartVertexID, NumVertices, Parameters, PropName, DataType.getExtent());
 						break;
 					case 3:
-						ConvertAlembicAttribute<Alembic::AbcGeom::IV3fGeomParam, Alembic::Abc::V3fArraySamplePtr, FVector>(HairDescription, StrandID, StartVertexID, NumVertices, Parameters, PropName, DataType.getExtent());
+						ConvertAlembicAttribute<Alembic::AbcGeom::IV3fGeomParam, Alembic::Abc::V3fArraySamplePtr, FVector>(HairDescription, StartStrandID, NumStrands, StartVertexID, NumVertices, Parameters, PropName, DataType.getExtent());
 						break;
 					}
 				}
@@ -291,13 +365,13 @@ namespace AlembicHairTranslatorUtils
 					switch (Extent)
 					{
 					case 1:
-						ConvertAlembicAttribute<Alembic::AbcGeom::IDoubleGeomParam, Alembic::Abc::DoubleArraySamplePtr, float>(HairDescription, StrandID, StartVertexID, NumVertices, Parameters, PropName);
+						ConvertAlembicAttribute<Alembic::AbcGeom::IDoubleGeomParam, Alembic::Abc::DoubleArraySamplePtr, float>(HairDescription, StartStrandID, NumStrands, StartVertexID, NumVertices, Parameters, PropName);
 						break;
 					case 2:
-						ConvertAlembicAttribute<Alembic::AbcGeom::IV2dGeomParam, Alembic::Abc::V2dArraySamplePtr, FVector2D>(HairDescription, StrandID, StartVertexID, NumVertices, Parameters, PropName, DataType.getExtent());
+						ConvertAlembicAttribute<Alembic::AbcGeom::IV2dGeomParam, Alembic::Abc::V2dArraySamplePtr, FVector2D>(HairDescription, StartStrandID, NumStrands, StartVertexID, NumVertices, Parameters, PropName, DataType.getExtent());
 						break;
 					case 3:
-						ConvertAlembicAttribute<Alembic::AbcGeom::IV3dGeomParam, Alembic::Abc::V3dArraySamplePtr, FVector>(HairDescription, StrandID, StartVertexID, NumVertices, Parameters, PropName, DataType.getExtent());
+						ConvertAlembicAttribute<Alembic::AbcGeom::IV3dGeomParam, Alembic::Abc::V3dArraySamplePtr, FVector>(HairDescription, StartStrandID, NumStrands, StartVertexID, NumVertices, Parameters, PropName, DataType.getExtent());
 						break;
 					}
 				}
@@ -321,23 +395,19 @@ FMatrix ConvertAlembicMatrix(const Alembic::Abc::M44d& AbcMatrix)
 
 static void ParseObject(const Alembic::Abc::IObject& InObject, FHairDescription& HairDescription, const FMatrix& ParentMatrix, const FMatrix& ConversionMatrix, float Scale, bool bCheckGroomAttributes)
 {
+	TRACE_CPUPROFILER_EVENT_SCOPE(ParseObject);
+
 	// Get MetaData info from current Alembic Object
 	const Alembic::Abc::MetaData ObjectMetaData = InObject.getMetaData();
 	const uint32 NumChildren = InObject.getNumChildren();
 
 	FMatrix LocalMatrix = ParentMatrix;
 
-	enum class EAttributeFrequency : uint8
-	{
-		None,
-		Groom,
-		Hair,
-		CV
-	};
-
 	bool bHandled = false;
 	if (Alembic::AbcGeom::ICurves::matches(ObjectMetaData))
 	{
+		TRACE_CPUPROFILER_EVENT_SCOPE(ParseICurves);
+
 		Alembic::AbcGeom::ICurves Curves = Alembic::AbcGeom::ICurves(InObject, Alembic::Abc::kWrapExisting);
 		Alembic::AbcGeom::ICurves::schema_type::Sample Sample = Curves.getSchema().getValue();
 
@@ -349,17 +419,9 @@ static void ParseObject(const Alembic::Abc::IObject& InObject, FHairDescription&
 		const uint32 NumPoints = Positions ? Positions->size() : 0;
 		const uint32 NumCurves = NumVertices->size(); // equivalent to Sample.getNumCurves()
 
-		EAttributeFrequency WidthFrequency = EAttributeFrequency::None;
-		{
-			if (NumWidths == NumPoints)
-			{
-				WidthFrequency = EAttributeFrequency::CV;
-			}
-			else if (NumWidths == NumCurves)
-			{
-				WidthFrequency = EAttributeFrequency::Hair;
-			}
-		}
+		// Get the starting strand and vertex IDs for this group of ICurves
+		int32 StartStrandID = HairDescription.GetNumStrands();
+		int32 StartVertexID = HairDescription.GetNumVertices();
 
 		FMatrix ConvertedMatrix = ParentMatrix * ConversionMatrix;
 		uint32 GlobalIndex = 0;
@@ -371,7 +433,6 @@ static void ParseObject(const Alembic::Abc::IObject& InObject, FHairDescription&
 
 			SetHairStrandAttribute(HairDescription, StrandID, HairAttribute::Strand::VertexCount, (int) CurveNumVertices);
 
-			int32 StartVertexID = HairDescription.GetNumVertices();
 			for (uint32 PointIndex = 0; PointIndex < CurveNumVertices; ++PointIndex, ++GlobalIndex)
 			{
 				FVertexID VertexID = HairDescription.AddVertex();
@@ -380,55 +441,59 @@ static void ParseObject(const Alembic::Abc::IObject& InObject, FHairDescription&
 
 				FVector ConvertedPosition = ConvertedMatrix.TransformPosition(FVector(Position.x, Position.y, Position.z));
 				SetHairVertexAttribute(HairDescription, VertexID, HairAttribute::Vertex::Position, ConvertedPosition);
+			}
+		}
 
-				float Width = 0;
-				switch (WidthFrequency)
-				{
-				case EAttributeFrequency::None:
-				{
-					const float CoordU = PointIndex / static_cast<float>(CurveNumVertices - 1);
-					Width = FMath::Lerp(AlembicHairFormat::RootRadius, AlembicHairFormat::TipRadius, CoordU);
-				}
-				break;
-				case EAttributeFrequency::CV:
-					Width = (*Widths)[GlobalIndex];
-					break;
-				}
+		// Set width values
+		// Determine the scope of the WidthsParam
+		Alembic::AbcGeom::GeometryScope WidthScope = Alembic::AbcGeom::kUnknownScope;
+		Alembic::AbcGeom::IFloatGeomParam WidthParam = Curves.getSchema().getWidthsParam();
+		if (WidthParam)
+		{
+			WidthScope = WidthParam.getScope();
+		}
 
-				// Per-vertex widths
-				if ((WidthFrequency == EAttributeFrequency::CV || WidthFrequency == EAttributeFrequency::None))
-				{
-					SetHairVertexAttribute(HairDescription, VertexID, HairAttribute::Vertex::Width, Width * Scale);
-				}
+		if (WidthScope == Alembic::AbcGeom::kConstantScope)
+		{
+			const float Width = (*Widths)[0] * Scale;
+			TStrandAttributesRef<float> WidthStrandAttributeRef = HairDescription.StrandAttributes().GetAttributesRef<float>(HairAttribute::Strand::Width);
+			if (!WidthStrandAttributeRef.IsValid())
+			{
+				HairDescription.StrandAttributes().RegisterAttribute<float>(HairAttribute::Strand::Width);
+				WidthStrandAttributeRef = HairDescription.StrandAttributes().GetAttributesRef<float>(HairAttribute::Strand::Width);
 			}
 
-			Alembic::AbcGeom::ICompoundProperty ArbParams = Curves.getSchema().getArbGeomParams();
-			if (ArbParams)
+			for (uint32 Index = 0; Index < NumCurves; ++Index)
 			{
-				AlembicHairTranslatorUtils::ConvertAlembicAttributes(HairDescription, StrandID, StartVertexID, CurveNumVertices, ArbParams);
+				WidthStrandAttributeRef[FStrandID(StartStrandID + Index)] = Width;
 			}
-
-			if (WidthFrequency == EAttributeFrequency::Hair)
+		}
+		else if (WidthScope == Alembic::AbcGeom::kVertexScope)
+		{
+			TVertexAttributesRef<float> VertexAttributeRef = HairDescription.VertexAttributes().GetAttributesRef<float>(HairAttribute::Vertex::Width);
+			if (!VertexAttributeRef.IsValid())
 			{
-				// Fallback if no per-strand or per-vertex groom_width attribute was found
-				TStrandAttributesRef<float> StrandWidths = HairDescription.StrandAttributes().GetAttributesRef<float>(HairAttribute::Strand::Width);
-				TVertexAttributesRef<float> VertexWidths = HairDescription.VertexAttributes().GetAttributesRef<float>(HairAttribute::Strand::Width);
+				HairDescription.VertexAttributes().RegisterAttribute<float>(HairAttribute::Vertex::Width);
+				VertexAttributeRef = HairDescription.VertexAttributes().GetAttributesRef<float>(HairAttribute::Vertex::Width);
+			}
 
-				if (!StrandWidths.IsValid() && !VertexWidths.IsValid())
-				{
-					const float Width = (*Widths)[CurveIndex];
-					SetHairStrandAttribute(HairDescription, StrandID, HairAttribute::Strand::Width, Width * Scale);
-				}
-				if (StrandWidths.IsValid() && StrandWidths[StrandID] == 0.f)
-				{
-					const float Width = (*Widths)[CurveIndex];
-					SetHairStrandAttribute(HairDescription, StrandID, HairAttribute::Strand::Width, Width * Scale);
-				}
+			for (uint32 Index = 0; Index < NumPoints; ++Index)
+			{
+				VertexAttributeRef[FVertexID(StartVertexID + Index)] = (*Widths)[Index] * Scale;
 			}
 		}
+
+		// Extract the arbitrary GeomParams here, only need to do it once per ICurves
+		Alembic::AbcGeom::ICompoundProperty ArbParams = Curves.getSchema().getArbGeomParams();
+		if (ArbParams)
+		{
+			AlembicHairTranslatorUtils::ConvertAlembicAttributes(HairDescription, StartStrandID, NumCurves, StartVertexID, NumPoints, ArbParams);
+		}
 	}
 	else if (Alembic::AbcGeom::IXform::matches(ObjectMetaData))
 	{
+		TRACE_CPUPROFILER_EVENT_SCOPE(ParseIXform);
+
 		Alembic::AbcGeom::IXform Xform = Alembic::AbcGeom::IXform(InObject, Alembic::Abc::kWrapExisting);
 		Alembic::AbcGeom::XformSample MatrixSample; 
 		Xform.getSchema().get(MatrixSample);
@@ -436,8 +501,20 @@ static void ParseObject(const Alembic::Abc::IObject& InObject, FHairDescription&
 		// The groom attributes should only be on the first IXform under the top node, no need to check for them once they are found
 		if (bCheckGroomAttributes)
 		{
+			// Groom attributes as UserProperties, 1.1+
+			Alembic::AbcGeom::ICompoundProperty Properties = Xform.getSchema().getUserProperties();
+			if (Properties)
+			{
+				if (Properties.getNumProperties() > 0)
+				{
+					AlembicHairTranslatorUtils::SetGroomAttributes(HairDescription, Properties);
+					bCheckGroomAttributes = false;
+				}
+			}
+
+			// Groom attributes as GeomParams, as fallback to 1.0
 			Alembic::AbcGeom::ICompoundProperty ArbParams = Xform.getSchema().getArbGeomParams();
-			if (ArbParams)
+			if (bCheckGroomAttributes && ArbParams)
 			{
 				if (ArbParams.getNumProperties() > 0)
 				{
@@ -461,6 +538,8 @@ static void ParseObject(const Alembic::Abc::IObject& InObject, FHairDescription&
 
 bool FAlembicHairTranslator::Translate(const FString& FileName, FHairDescription& HairDescription, const FGroomConversionSettings& ConversionSettings)
 {
+	TRACE_CPUPROFILER_EVENT_SCOPE(FAlembicHairTranslator::Translate);
+
 	/** Factory used to generate objects*/
 	Alembic::AbcCoreFactory::IFactory Factory;
 	Alembic::AbcCoreFactory::IFactory::CoreType CompressionType = Alembic::AbcCoreFactory::IFactory::kUnknown;
