Updated Alembic hair translator to 1.3:
* Added support for UniformScope hair widths.
* Added support for ConstantScope ArbGeomParams as strand attribute.

diff --git a/Engine/Plugins/Experimental/AlembicHairImporter/Source/AlembicHairTranslator/Private/AlembicHairTranslator.cpp b/Engine/Plugins/Experimental/AlembicHairImporter/Source/AlembicHairTranslator/Private/AlembicHairTranslator.cpp
--- a/Engine/Plugins/Experimental/AlembicHairImporter/Source/AlembicHairTranslator/Private/AlembicHairTranslator.cpp
+++ b/Engine/Plugins/Experimental/AlembicHairImporter/Source/AlembicHairTranslator/Private/AlembicHairTranslator.cpp
@@ -213,9 +213,8 @@ namespace AlembicHairTranslatorUtils
 		Alembic::AbcGeom::GeometryScope Scope = Param.getScope();
 		int32 NumValues = ParamValues->size();
 
-		// Check the supported scope: UniformScope or VertexScope (1.2)
-		// ConstantScope (1.0) not useful so use NumStrands or NumVertices to determine scope
-		if (Scope == Alembic::AbcGeom::kUniformScope || NumValues == NumStrands)
+		// Check the supported scope: UniformScope or VertexScope (1.2), ConstantScope (1.3)
+		if (Scope == Alembic::AbcGeom::kUniformScope || Scope == Alembic::AbcGeom::kConstantScope || NumValues == NumStrands)
 		{
 			TStrandAttributesRef<AttributeType> StrandAttributeRef = HairDescription.StrandAttributes().GetAttributesRef<AttributeType>(AttributeName);
 			if (!StrandAttributeRef.IsValid())
@@ -224,9 +223,12 @@ namespace AlembicHairTranslatorUtils
 				StrandAttributeRef = HairDescription.StrandAttributes().GetAttributesRef<AttributeType>(AttributeName);
 			}
 
-			for (int32 StrandIndex = 0; StrandIndex < NumValues; ++StrandIndex)
+			// Additional check in case the scope is wrongly set as ConstantScope but the number of values match the number of strands
+			// Treat that case as UniformScope. Note that with ConstantScope, NumValues is 1, thus the reason to iterate on NumStrands
+			bool bIsConstantValue = Scope == Alembic::AbcGeom::kConstantScope && NumValues != NumStrands;
+			for (int32 StrandIndex = 0; StrandIndex < NumStrands; ++StrandIndex)
 			{
-				StrandAttributeRef[FStrandID(StartStrandID + StrandIndex)] = (*ParamValues)[StrandIndex];
+				StrandAttributeRef[FStrandID(StartStrandID + StrandIndex)] = (*ParamValues)[bIsConstantValue ? 0 : StrandIndex];
 			}
 		}
 		else if (Scope == Alembic::AbcGeom::kVertexScope || NumValues == NumVertices)
@@ -256,9 +258,8 @@ namespace AlembicHairTranslatorUtils
 		Alembic::AbcGeom::GeometryScope Scope = Param.getScope();
 		int32 NumValues = ParamValues->size();
 
-		// Check the supported scope: UniformScope or VertexScope (1.2)
-		// ConstantScope (1.0) not useful so use NumStrands or NumVertices to determine scope
-		if (Scope == Alembic::AbcGeom::kUniformScope || NumValues == NumStrands)
+		// Check the supported scope: UniformScope or VertexScope (1.2), ConstantScope (1.3)
+		if (Scope == Alembic::AbcGeom::kUniformScope || Scope == Alembic::AbcGeom::kConstantScope || NumValues == NumStrands)
 		{
 			TStrandAttributesRef<AttributeType> StrandAttributeRef = HairDescription.StrandAttributes().GetAttributesRef<AttributeType>(AttributeName);
 			if (!StrandAttributeRef.IsValid())
@@ -267,12 +268,15 @@ namespace AlembicHairTranslatorUtils
 				StrandAttributeRef = HairDescription.StrandAttributes().GetAttributesRef<AttributeType>(AttributeName);
 			}
 
-			for (int32 StrandIndex = 0; StrandIndex < NumValues; ++StrandIndex)
+			// Additional check in case the scope is wrongly set as ConstantScope but the number of values match the number of strands
+			// Treat that case as UniformScope. Note that with ConstantScope, NumValues is 1, thus the reason to iterate on NumStrands
+			bool bIsConstantValue = Scope == Alembic::AbcGeom::kConstantScope && NumValues != NumStrands;
+			for (int32 StrandIndex = 0; StrandIndex < NumStrands; ++StrandIndex)
 			{
 				AttributeType ParamValue;
 				for (int32 Index = 0; Index < Extent; ++Index)
 				{
-					ParamValue[Index] = (*ParamValues)[StrandIndex][Index];
+					ParamValue[Index] = (*ParamValues)[bIsConstantValue ? 0 : StrandIndex][Index];
 				}
 
 				StrandAttributeRef[FStrandID(StartStrandID + StrandIndex)] = ParamValue;
@@ -453,9 +457,9 @@ static void ParseObject(const Alembic::Abc::IObject& InObject, FHairDescription&
 			WidthScope = WidthParam.getScope();
 		}
 
-		if (WidthScope == Alembic::AbcGeom::kConstantScope)
+		if (WidthScope == Alembic::AbcGeom::kConstantScope || WidthScope == Alembic::AbcGeom::kUniformScope)
 		{
-			const float Width = (*Widths)[0] * Scale;
+			const float ConstWidth = (*Widths)[0] * Scale;
 			TStrandAttributesRef<float> WidthStrandAttributeRef = HairDescription.StrandAttributes().GetAttributesRef<float>(HairAttribute::Strand::Width);
 			if (!WidthStrandAttributeRef.IsValid())
 			{
@@ -465,7 +469,7 @@ static void ParseObject(const Alembic::Abc::IObject& InObject, FHairDescription&
 
 			for (uint32 Index = 0; Index < NumCurves; ++Index)
 			{
-				WidthStrandAttributeRef[FStrandID(StartStrandID + Index)] = Width;
+				WidthStrandAttributeRef[FStrandID(StartStrandID + Index)] = WidthScope == Alembic::AbcGeom::kConstantScope ? ConstWidth : (*Widths)[Index] * Scale;
 			}
 		}
 		else if (WidthScope == Alembic::AbcGeom::kVertexScope)
