Updated Alembic hair translator to 1.3:
* Added support for reading UserProperties on ICurves for groom-scoped attributes.

diff --git a/Engine/Plugins/Experimental/AlembicHairImporter/Source/AlembicHairTranslator/Private/AlembicHairTranslator.cpp b/Engine/Plugins/Experimental/AlembicHairImporter/Source/AlembicHairTranslator/Private/AlembicHairTranslator.cpp
--- a/Engine/Plugins/Experimental/AlembicHairImporter/Source/AlembicHairTranslator/Private/AlembicHairTranslator.cpp
+++ b/Engine/Plugins/Experimental/AlembicHairImporter/Source/AlembicHairTranslator/Private/AlembicHairTranslator.cpp
@@ -493,6 +493,20 @@ static void ParseObject(const Alembic::Abc::IObject& InObject, FHairDescription&
 		{
 			AlembicHairTranslatorUtils::ConvertAlembicAttributes(HairDescription, StartStrandID, NumCurves, StartVertexID, NumPoints, ArbParams);
 		}
+
+		if (bCheckGroomAttributes)
+		{
+			// Groom attributes as UserProperties on ICurves, 1.3+
+			Alembic::AbcGeom::ICompoundProperty Properties = Curves.getSchema().getUserProperties();
+			if (Properties)
+			{
+				if (Properties.getNumProperties() > 0)
+				{
+					AlembicHairTranslatorUtils::SetGroomAttributes(HairDescription, Properties);
+					bCheckGroomAttributes = false;
+				}
+			}
+		}
 	}
 	else if (Alembic::AbcGeom::IXform::matches(ObjectMetaData))
 	{
