Fixed crash during Alembic groom import when there's a mismatch in the number of vertices in the curves node.

diff --git a/Engine/Plugins/Experimental/AlembicHairImporter/Source/AlembicHairTranslator/Private/AlembicHairTranslator.cpp b/Engine/Plugins/Experimental/AlembicHairImporter/Source/AlembicHairTranslator/Private/AlembicHairTranslator.cpp
--- a/Engine/Plugins/Experimental/AlembicHairImporter/Source/AlembicHairTranslator/Private/AlembicHairTranslator.cpp
+++ b/Engine/Plugins/Experimental/AlembicHairImporter/Source/AlembicHairTranslator/Private/AlembicHairTranslator.cpp
@@ -420,8 +420,8 @@ static void ParseObject(const Alembic::Abc::IObject& InObject, FHairDescription&
 		Alembic::Abc::Int32ArraySamplePtr NumVertices = Sample.getCurvesNumVertices();
 
 		const int32 NumWidths = Widths ? Widths->size() : 0;
-		const uint32 NumPoints = Positions ? Positions->size() : 0;
-		const uint32 NumCurves = NumVertices->size(); // equivalent to Sample.getNumCurves()
+		uint32 NumPoints = Positions ? Positions->size() : 0;
+		uint32 NumCurves = NumVertices ? NumVertices->size() : 0; // equivalent to Sample.getNumCurves()
 
 		// Get the starting strand and vertex IDs for this group of ICurves
 		int32 StartStrandID = HairDescription.GetNumStrands();
@@ -429,10 +429,24 @@ static void ParseObject(const Alembic::Abc::IObject& InObject, FHairDescription&
 
 		FMatrix ConvertedMatrix = ParentMatrix * ConversionMatrix;
 		uint32 GlobalIndex = 0;
+		uint32 TotalVertices = 0;
 		for (uint32 CurveIndex = 0; CurveIndex < NumCurves; ++CurveIndex)
 		{
 			const uint32 CurveNumVertices = (*NumVertices)[CurveIndex];
 
+			// Check the running total number of vertices and skip processing the rest of the node if there is mismatch with the number of points
+			TotalVertices += CurveNumVertices;
+			if (TotalVertices > NumPoints)
+			{
+				UE_LOG(LogAlembicHairImporter, Warning, TEXT("Curve %u of %u has %u vertices which causes total vertices (%u) to exceed the expected vertices (%u) in ICurves node. This curve and the remaining ones in the node will be skipped."),
+					CurveIndex + 1, NumCurves, CurveNumVertices, TotalVertices, NumPoints);
+
+				// Adjust the number of curves and points to those processed so far
+				NumCurves = HairDescription.GetNumStrands() - StartStrandID;
+				NumPoints = HairDescription.GetNumVertices() - StartVertexID;
+				break;
+			}
+
 			FStrandID StrandID = HairDescription.AddStrand();
 
 			SetHairStrandAttribute(HairDescription, StrandID, HairAttribute::Strand::VertexCount, (int) CurveNumVertices);
