Fixed Alembic GroomCache sampling to sample by time instead of by index.

diff --git a/Engine/Plugins/Importers/AlembicHairImporter/Source/AlembicHairTranslator/Private/AlembicHairTranslator.cpp b/Engine/Plugins/Importers/AlembicHairImporter/Source/AlembicHairTranslator/Private/AlembicHairTranslator.cpp
--- a/Engine/Plugins/Importers/AlembicHairImporter/Source/AlembicHairTranslator/Private/AlembicHairTranslator.cpp
+++ b/Engine/Plugins/Importers/AlembicHairImporter/Source/AlembicHairTranslator/Private/AlembicHairTranslator.cpp
@@ -425,7 +425,7 @@ FMatrix ConvertAlembicMatrix(const Alembic::Abc::M44d& AbcMatrix)
 	return Matrix;
 }
 
-static void ParseObject(const Alembic::Abc::IObject& InObject, int32 FrameIndex, FHairDescription& HairDescription, const FMatrix& ParentMatrix, const FMatrix& ConversionMatrix, float Scale, bool bCheckGroomAttributes, FGroomAnimationInfo* AnimInfo)
+static void ParseObject(const Alembic::Abc::IObject& InObject, float FrameTime, FHairDescription& HairDescription, const FMatrix& ParentMatrix, const FMatrix& ConversionMatrix, float Scale, bool bCheckGroomAttributes, FGroomAnimationInfo* AnimInfo)
 {
 	TRACE_CPUPROFILER_EVENT_SCOPE(ParseObject);
 
@@ -490,7 +490,7 @@ static void ParseObject(const Alembic::Abc::IObject& InObject, int32 FrameIndex,
 			}
 		}
 
-		Alembic::Abc::ISampleSelector SampleSelector((Alembic::Abc::index_t) FrameIndex);
+		Alembic::Abc::ISampleSelector SampleSelector((Alembic::Abc::chrono_t) FrameTime);
 		Alembic::AbcGeom::ICurves::schema_type::Sample Sample = Curves.getSchema().getValue(SampleSelector);
 
 		Alembic::Abc::FloatArraySamplePtr Widths = Curves.getSchema().getWidthsParam() ? Curves.getSchema().getWidthsParam().getExpandedValue(SampleSelector).getVals() : nullptr;
@@ -641,7 +641,7 @@ static void ParseObject(const Alembic::Abc::IObject& InObject, int32 FrameIndex,
 	{
 		for (uint32 ChildIndex = 0; ChildIndex < NumChildren; ++ChildIndex)
 		{
-			ParseObject(InObject.getChild(ChildIndex), FrameIndex, HairDescription, LocalMatrix, ConversionMatrix, Scale, bCheckGroomAttributes, AnimInfo);
+			ParseObject(InObject.getChild(ChildIndex), FrameTime, HairDescription, LocalMatrix, ConversionMatrix, Scale, bCheckGroomAttributes, AnimInfo);
 		}
 	}
 }
@@ -693,7 +693,7 @@ bool FAlembicHairTranslator::Translate(const FString& FileName, FHairDescription
 	FMatrix ConversionMatrix = FScaleMatrix::Make(ConversionSettings.Scale) * FRotationMatrix::Make(FQuat::MakeFromEuler(ConversionSettings.Rotation));
 	FMatrix ParentMatrix = FMatrix::Identity;
 	const float StrandsWidthScale = FMath::Abs(ConversionSettings.Scale.X); // Assume uniform scaling
-	ParseObject(TopObject, 0, HairDescription, ParentMatrix, ConversionMatrix, StrandsWidthScale, true, OutAnimInfo);
+	ParseObject(TopObject, 0.0f, HairDescription, ParentMatrix, ConversionMatrix, StrandsWidthScale, true, OutAnimInfo);
 
 	return HairDescription.IsValid();
 }
@@ -741,7 +741,7 @@ bool FAlembicHairTranslator::BeginTranslation(const FString& FileName)
 	return Abc->IsValid();
 }
 
-bool FAlembicHairTranslator::Translate(uint32 FrameIndex, FHairDescription& HairDescription, const struct FGroomConversionSettings& ConversionSettings)
+bool FAlembicHairTranslator::Translate(float FrameTime, FHairDescription& HairDescription, const struct FGroomConversionSettings& ConversionSettings)
 {
 	if (!Abc || !Abc->IsValid())
 	{
@@ -751,7 +751,7 @@ bool FAlembicHairTranslator::Translate(uint32 FrameIndex, FHairDescription& Hair
 	FMatrix ConversionMatrix = FScaleMatrix::Make(ConversionSettings.Scale) * FRotationMatrix::Make(FQuat::MakeFromEuler(ConversionSettings.Rotation));
 	FMatrix ParentMatrix = FMatrix::Identity;
 	const float StrandsWidthScale = FMath::Abs(ConversionSettings.Scale.X); // Assume uniform scaling
-	ParseObject(Abc->TopObject, FrameIndex, HairDescription, ParentMatrix, ConversionMatrix, StrandsWidthScale, true, nullptr);
+	ParseObject(Abc->TopObject, FrameTime, HairDescription, ParentMatrix, ConversionMatrix, StrandsWidthScale, true, nullptr);
 
 	return HairDescription.IsValid();
 }
diff --git a/Engine/Plugins/Importers/AlembicHairImporter/Source/AlembicHairTranslator/Private/AlembicHairTranslator.h b/Engine/Plugins/Importers/AlembicHairImporter/Source/AlembicHairTranslator/Private/AlembicHairTranslator.h
--- a/Engine/Plugins/Importers/AlembicHairImporter/Source/AlembicHairTranslator/Private/AlembicHairTranslator.h
+++ b/Engine/Plugins/Importers/AlembicHairImporter/Source/AlembicHairTranslator/Private/AlembicHairTranslator.h
@@ -17,7 +17,7 @@ public:
 
 	virtual bool Translate(const FString& FilePath, FHairDescription& OutHairDescription, const struct FGroomConversionSettings& ConversionSettings, FGroomAnimationInfo* OutAnimInfo) override;
 	virtual bool BeginTranslation(const FString& FilePath) override;
-	virtual bool Translate(uint32 FrameIndex, FHairDescription& OutHairDescription, const struct FGroomConversionSettings& ConversionSettings) override;
+	virtual bool Translate(float FrameTime, FHairDescription& OutHairDescription, const struct FGroomConversionSettings& ConversionSettings) override;
 	virtual void EndTranslation() override;
 
 private:
diff --git a/Engine/Plugins/Runtime/HairStrands/Source/HairStrandsEditor/Private/GroomCacheImporter.cpp b/Engine/Plugins/Runtime/HairStrands/Source/HairStrandsEditor/Private/GroomCacheImporter.cpp
--- a/Engine/Plugins/Runtime/HairStrands/Source/HairStrandsEditor/Private/GroomCacheImporter.cpp
+++ b/Engine/Plugins/Runtime/HairStrands/Source/HairStrandsEditor/Private/GroomCacheImporter.cpp
@@ -117,7 +117,7 @@ TArray<UGroomCache*> FGroomCacheImporter::ImportGroomCache(const FString& Source
 			SlowTask.EnterProgressFrame(1.f, TextBuilder.ToText());
 
 			FHairDescription FrameHairDescription;
-			if (Translator->Translate(FrameIndex, FrameHairDescription, HairImportContext.ImportOptions->ConversionSettings))
+			if (Translator->Translate(FrameIndex * AnimInfo.SecondsPerFrame, FrameHairDescription, HairImportContext.ImportOptions->ConversionSettings))
 			{
 				FHairDescriptionGroups HairDescriptionGroups;
 				if (!FGroomBuilder::BuildHairDescriptionGroups(FrameHairDescription, HairDescriptionGroups))
diff --git a/Engine/Plugins/Runtime/HairStrands/Source/HairStrandsEditor/Public/HairStrandsTranslator.h b/Engine/Plugins/Runtime/HairStrands/Source/HairStrandsEditor/Public/HairStrandsTranslator.h
--- a/Engine/Plugins/Runtime/HairStrands/Source/HairStrandsEditor/Public/HairStrandsTranslator.h
+++ b/Engine/Plugins/Runtime/HairStrands/Source/HairStrandsEditor/Public/HairStrandsTranslator.h
@@ -31,8 +31,12 @@ public:
 	virtual bool BeginTranslation(const FString& FilePath) { return false; }
 
 	/** Translate a given file into a HairDescription at the requested frame index; return true if successful */
+	UE_DEPRECATED(5.0, "Translate by frame index is deprecated. Use Translate by frame time instead.")
 	virtual bool Translate(uint32 FrameIndex, FHairDescription& OutHairDescription, const struct FGroomConversionSettings& ConversionSettings) { return false; }
 
+	/** Translate a given file into a HairDescription at the requested frame time (in seconds); return true if successful */
+	virtual bool Translate(float FrameTime, FHairDescription& OutHairDescription, const struct FGroomConversionSettings& ConversionSettings) { return false; }
+
 	/** Clean up after finishing translations */
 	virtual void EndTranslation() { }
 };
