Amending alembic hair serialization -- serializing in interpolation data to save along with the UE grrom asset file (needed for hair-card generation and for future interpolation efforts).

diff --git a/Engine/Plugins/Importers/AlembicHairImporter/Source/AlembicHairTranslator/AlembicHairTranslatorModule.Build.cs b/Engine/Plugins/Importers/AlembicHairImporter/Source/AlembicHairTranslator/AlembicHairTranslatorModule.Build.cs
--- a/Engine/Plugins/Importers/AlembicHairImporter/Source/AlembicHairTranslator/AlembicHairTranslatorModule.Build.cs
+++ b/Engine/Plugins/Importers/AlembicHairImporter/Source/AlembicHairTranslator/AlembicHairTranslatorModule.Build.cs
@@ -13,7 +13,8 @@ namespace UnrealBuildTool.Rules
                     "Core",
                     "Engine",
                     "HairStrandsCore",
-                    "HairStrandsEditor"
+                    "HairStrandsEditor",
+					"CoreUObject"
                 });
         }
 	}
diff --git a/Engine/Plugins/Importers/AlembicHairImporter/Source/AlembicHairTranslator/Private/AlembicHairTranslator.cpp b/Engine/Plugins/Importers/AlembicHairImporter/Source/AlembicHairTranslator/Private/AlembicHairTranslator.cpp
--- a/Engine/Plugins/Importers/AlembicHairImporter/Source/AlembicHairTranslator/Private/AlembicHairTranslator.cpp
+++ b/Engine/Plugins/Importers/AlembicHairImporter/Source/AlembicHairTranslator/Private/AlembicHairTranslator.cpp
@@ -109,6 +109,12 @@ namespace AlembicHairTranslatorUtils
 		}
 	}
 
+	template <typename AttributeType>
+	void SetGroomAttribute(FHairDescription& HairDescription, FName AttributeName, AttributeType AttributeValue)
+	{
+		SetGroomAttribute(HairDescription, FGroomID(0), AttributeName, AttributeValue);
+	}
+
 	void SetGroomAttributes(FHairDescription& HairDescription, const Alembic::AbcGeom::ICompoundProperty& Parameters)
 	{
 		for (int Index = 0; Index < Parameters.getNumProperties(); ++Index)
@@ -203,6 +209,80 @@ namespace AlembicHairTranslatorUtils
 		}
 	}
 
+	EGroomBasisType MapAlembicEnumToUe(Alembic::AbcGeom::BasisType AlembicBasisType, bool& bSuccess)
+	{
+		bSuccess = true;
+		switch (AlembicBasisType)
+		{
+		case Alembic::AbcGeom::kNoBasis:
+			return EGroomBasisType::NoBasis;
+		case Alembic::AbcGeom::kBezierBasis:
+			return EGroomBasisType::BezierBasis;
+		case Alembic::AbcGeom::kBsplineBasis:
+			return EGroomBasisType::BsplineBasis;
+		case Alembic::AbcGeom::kCatmullromBasis:
+			return EGroomBasisType::CatmullromBasis;
+		case Alembic::AbcGeom::kHermiteBasis:
+			return EGroomBasisType::HermiteBasis;
+		case Alembic::AbcGeom::kPowerBasis:
+			return EGroomBasisType::PowerBasis;
+		default:
+			ensureMsgf(false, TEXT("Unsupported basis type: %d"), AlembicBasisType);
+			break;
+		}
+
+		bSuccess = false;
+		return EGroomBasisType::NoBasis;
+	}
+
+	EGroomCurveType MapAlembicEnumToUe(Alembic::AbcGeom::CurveType AlembicCurveType, bool& bSuccess)
+	{
+		bSuccess = true;
+		switch (AlembicCurveType)
+		{
+		case Alembic::AbcGeom::kCubic:
+			return EGroomCurveType::Cubic;
+		case Alembic::AbcGeom::kLinear:
+			return EGroomCurveType::Linear;
+		case Alembic::AbcGeom::kVariableOrder:
+			return EGroomCurveType::VariableOrder;
+		default:
+			ensureMsgf(false, TEXT("Unsupported curve type: %d"), AlembicCurveType);
+			break;
+		}
+
+		bSuccess = false;
+		return EGroomCurveType::Linear;
+	}
+
+	template<typename AlembicEnumType>
+	bool ConvertAlembicEnumToStrandAttribute(FHairDescription& HairDescription, FStrandID StrandID, AlembicEnumType AlembicEnumValue, FName AttributeName)
+	{
+		bool bSuccess = false;
+		auto UeValue = MapAlembicEnumToUe(AlembicEnumValue, bSuccess);
+		if (bSuccess)
+		{
+			UEnum* UeEnum = StaticEnum<decltype(UeValue)>();
+			if (ensure(UeEnum))
+			{
+				FName ValueName = *UeEnum->GetNameStringByValue((int)UeValue);
+				if (ValueName.IsNone())
+				{
+					bSuccess = false;
+				}
+				else
+				{
+					SetHairStrandAttribute(HairDescription, StrandID, AttributeName, ValueName);
+				}
+			}
+			else
+			{
+				bSuccess = false;
+			}
+		}
+		return bSuccess;
+	}
+
 	template <typename AbcParamType, typename AbcArraySampleType, typename AttributeType>
 	void ConvertAlembicAttribute(FHairDescription& HairDescription, int32 StartStrandID, int32 NumStrands, int32 StartVertexID, int32 NumVertices, const Alembic::AbcGeom::ICompoundProperty& Parameters, const std::string& PropName)
 	{
@@ -496,10 +576,16 @@ static void ParseObject(const Alembic::Abc::IObject& InObject, float FrameTime,
 		Alembic::Abc::FloatArraySamplePtr Widths = Curves.getSchema().getWidthsParam() ? Curves.getSchema().getWidthsParam().getExpandedValue(SampleSelector).getVals() : nullptr;
 		Alembic::Abc::P3fArraySamplePtr Positions = Sample.getPositions();
 		Alembic::Abc::Int32ArraySamplePtr NumVertices = Sample.getCurvesNumVertices();
+		Alembic::Abc::FloatArraySamplePtr Knots = Sample.getKnots();
+		Alembic::Abc::UcharArraySamplePtr Orders = Sample.getOrders();
+
+		Alembic::AbcGeom::BasisType BasisType = Sample.getBasis();
+		Alembic::AbcGeom::CurveType CurveType = Sample.getType();
 
 		const int32 NumWidths = Widths ? Widths->size() : 0;
 		uint32 NumPoints = Positions ? Positions->size() : 0;
 		uint32 NumCurves = NumVertices ? NumVertices->size() : 0; // equivalent to Sample.getNumCurves()
+		const uint32 NumKnots = Knots ? Knots->size() : 0;
 
 		// Get the starting strand and vertex IDs for this group of ICurves
 		int32 StartStrandID = HairDescription.GetNumStrands();
@@ -529,6 +615,9 @@ static void ParseObject(const Alembic::Abc::IObject& InObject, float FrameTime,
 
 			SetHairStrandAttribute(HairDescription, StrandID, HairAttribute::Strand::VertexCount, (int) CurveNumVertices);
 
+			AlembicHairTranslatorUtils::ConvertAlembicEnumToStrandAttribute(HairDescription, StrandID, BasisType, HairAttribute::Strand::BasisType);
+			AlembicHairTranslatorUtils::ConvertAlembicEnumToStrandAttribute(HairDescription, StrandID, CurveType, HairAttribute::Strand::CurveType);
+
 			for (uint32 PointIndex = 0; PointIndex < CurveNumVertices; ++PointIndex, ++GlobalIndex)
 			{
 				FVertexID VertexID = HairDescription.AddVertex();
@@ -579,6 +668,53 @@ static void ParseObject(const Alembic::Abc::IObject& InObject, float FrameTime,
 			}
 		}
 
+		if (BasisType != Alembic::AbcGeom::kNoBasis && Knots)
+		{
+			TStrandAttributesRef<TArrayAttribute<float>> KnotsStrandAttributeRef = HairDescription.StrandAttributes().GetAttributesRef<TArrayAttribute<float>>(HairAttribute::Strand::Knots);
+			if (!KnotsStrandAttributeRef.IsValid())
+			{
+				HairDescription.StrandAttributes().RegisterAttribute<float[]>(HairAttribute::Strand::Knots);
+				KnotsStrandAttributeRef = HairDescription.StrandAttributes().GetAttributesRef<TArrayAttribute<float>>(HairAttribute::Strand::Knots);
+			}		
+
+			uint32 GlobalKnotIndex = 0;
+			for (uint32 CurveIndex = 0; CurveIndex < NumCurves; ++CurveIndex)
+			{
+				uint32 DegreeOfCurve = 1;
+				if (CurveType == Alembic::AbcGeom::kCubic)
+				{
+					DegreeOfCurve = 3;
+				}
+				else if (CurveType == Alembic::AbcGeom::kVariableOrder)
+				{
+					if (ensureMsgf(Orders && CurveIndex < Orders->size(), TEXT("`Orders` index out of bounds, skipping knots for this strand.")))
+					{
+						DegreeOfCurve = (*Orders)[CurveIndex];
+					}
+					else
+					{
+						continue;
+					}
+				}
+
+				TArrayAttribute<float> KnotsArrayAttribute = KnotsStrandAttributeRef[FStrandID(StartStrandID + CurveIndex)];
+
+				uint32 CurveNumKnots = (*NumVertices)[CurveIndex] + DegreeOfCurve + 1;
+				if (!ensureMsgf(GlobalKnotIndex + CurveNumKnots <= NumKnots, TEXT("Computed number of knots would exceed the number available from the groom.")))
+				{
+					CurveNumKnots = NumKnots - GlobalKnotIndex;
+				}
+				KnotsArrayAttribute.SetNum(CurveNumKnots);
+
+				for (uint32 KnotIndex = 0; KnotIndex < CurveNumKnots; ++KnotIndex, ++GlobalKnotIndex)
+				{
+					KnotsArrayAttribute[KnotIndex] = (*Knots)[GlobalKnotIndex];
+				}
+			}
+
+			ensureMsgf(GlobalKnotIndex == NumKnots, TEXT("Failed to translate all knots for the groom."));
+		}
+
 		// Extract the arbitrary GeomParams here, only need to do it once per ICurves
 		Alembic::AbcGeom::ICompoundProperty ArbParams = Curves.getSchema().getArbGeomParams();
 		if (ArbParams)
diff --git a/Engine/Plugins/Runtime/HairStrands/Source/HairStrandsCore/Private/HairAttributes.cpp b/Engine/Plugins/Runtime/HairStrands/Source/HairStrandsCore/Private/HairAttributes.cpp
--- a/Engine/Plugins/Runtime/HairStrands/Source/HairStrandsCore/Private/HairAttributes.cpp
+++ b/Engine/Plugins/Runtime/HairStrands/Source/HairStrandsCore/Private/HairAttributes.cpp
@@ -19,6 +19,9 @@ namespace HairAttribute
 	const FName Strand::Width("groom_width");
 	const FName Strand::ClosestGuides("groom_closest_guides");
 	const FName Strand::GuideWeights("groom_guide_weights");
+	const FName Strand::BasisType("groom_basis_type");
+	const FName Strand::CurveType("groom_curve_type");
+	const FName Strand::Knots("groom_knots");
 
 	const FName Groom::Color("groom_color");
 	const FName Groom::Roughness("groom_roughness");
diff --git a/Engine/Plugins/Runtime/HairStrands/Source/HairStrandsCore/Public/GroomCacheData.h b/Engine/Plugins/Runtime/HairStrands/Source/HairStrandsCore/Public/GroomCacheData.h
--- a/Engine/Plugins/Runtime/HairStrands/Source/HairStrandsCore/Public/GroomCacheData.h
+++ b/Engine/Plugins/Runtime/HairStrands/Source/HairStrandsCore/Public/GroomCacheData.h
@@ -71,6 +71,25 @@ enum class EGroomCacheType : uint8
 	Guides
 };
 
+UENUM(BlueprintType)
+enum class EGroomBasisType : uint8
+{
+	NoBasis,
+	BezierBasis,
+	BsplineBasis,
+	CatmullromBasis,
+	HermiteBasis,
+	PowerBasis,
+};
+
+UENUM(BlueprintType)
+enum class EGroomCurveType : uint8
+{
+	Cubic,
+	Linear,
+	VariableOrder,
+};
+
 /** Information about the GroomCache itself */
 USTRUCT()
 struct FGroomCacheInfo
diff --git a/Engine/Plugins/Runtime/HairStrands/Source/HairStrandsCore/Public/HairAttributes.h b/Engine/Plugins/Runtime/HairStrands/Source/HairStrandsCore/Public/HairAttributes.h
--- a/Engine/Plugins/Runtime/HairStrands/Source/HairStrandsCore/Public/HairAttributes.h
+++ b/Engine/Plugins/Runtime/HairStrands/Source/HairStrandsCore/Public/HairAttributes.h
@@ -26,6 +26,9 @@ namespace HairAttribute
 		extern HAIRSTRANDSCORE_API const FName Width;			// float
 		extern HAIRSTRANDSCORE_API const FName ClosestGuides;	// FVector
 		extern HAIRSTRANDSCORE_API const FName GuideWeights;	// FVector
+		extern HAIRSTRANDSCORE_API const FName BasisType;		// FName (EGroomBasisType)
+		extern HAIRSTRANDSCORE_API const FName CurveType;		// FName (EGroomCurveType)
+		extern HAIRSTRANDSCORE_API const FName Knots;			// float[]
 	}
 
 	namespace Groom
