Fix alembic string translation when importing groom.

This fixes missing group name on grooms.

diff --git a/Engine/Plugins/Importers/AlembicHairImporter/Source/AlembicHairTranslator/Private/AlembicHairTranslator.cpp b/Engine/Plugins/Importers/AlembicHairImporter/Source/AlembicHairTranslator/Private/AlembicHairTranslator.cpp
--- a/Engine/Plugins/Importers/AlembicHairImporter/Source/AlembicHairTranslator/Private/AlembicHairTranslator.cpp
+++ b/Engine/Plugins/Importers/AlembicHairImporter/Source/AlembicHairTranslator/Private/AlembicHairTranslator.cpp
@@ -209,6 +209,27 @@ namespace AlembicHairTranslatorUtils
 		}
 	}
 
+	// Default conversion traits
+	template<typename TAttributeType>
+	struct TConvertor
+	{
+		template<typename InType>
+		static TAttributeType ConvertType(const InType& In)
+		{
+			return In;
+		}
+	};
+	// String conversion traits
+	template<>
+	struct TConvertor<FName>
+	{
+		template<typename InType>
+		static FName ConvertType(const InType& In)
+		{
+			return FName(ANSI_TO_TCHAR(In.c_str()));
+		}
+	};
+
 	EGroomBasisType MapAlembicEnumToUe(Alembic::AbcGeom::BasisType AlembicBasisType, bool& bSuccess)
 	{
 		bSuccess = true;
@@ -309,7 +330,7 @@ namespace AlembicHairTranslatorUtils
 			bool bIsConstantValue = Scope == Alembic::AbcGeom::kConstantScope && NumValues != NumStrands;
 			for (int32 StrandIndex = 0; StrandIndex < NumStrands; ++StrandIndex)
 			{
-				StrandAttributeRef[FStrandID(StartStrandID + StrandIndex)] = (*ParamValues)[bIsConstantValue ? 0 : StrandIndex];
+				StrandAttributeRef[FStrandID(StartStrandID + StrandIndex)] = TConvertor<AttributeType>::ConvertType((*ParamValues)[bIsConstantValue ? 0 : StrandIndex]);
 			}
 		}
 		else if (Scope == Alembic::AbcGeom::kVertexScope || NumValues == NumVertices)
@@ -323,7 +344,7 @@ namespace AlembicHairTranslatorUtils
 
 			for (int32 VertexIndex = 0; VertexIndex < NumVertices; ++VertexIndex)
 			{
-				VertexAttributeRef[FVertexID(StartVertexID + VertexIndex)] = (*ParamValues)[VertexIndex];
+				VertexAttributeRef[FVertexID(StartVertexID + VertexIndex)] = TConvertor<AttributeType>::ConvertType((*ParamValues)[VertexIndex]);
 			}
 		}
 	}
@@ -428,6 +449,11 @@ namespace AlembicHairTranslatorUtils
 
 				switch (DataType.getPod())
 				{
+				case Alembic::Util::kStringPOD:
+				{
+					ConvertAlembicAttribute<Alembic::AbcGeom::IStringGeomParam, Alembic::Abc::StringArraySamplePtr, FName>(HairDescription, StartStrandID, NumStrands, StartVertexID, NumVertices, Parameters, PropName);
+				}
+				break;
 				case Alembic::Util::kBooleanPOD:
 				{
 					ConvertAlembicAttribute<Alembic::AbcGeom::IBoolGeomParam, Alembic::Abc::BoolArraySamplePtr, bool>(HairDescription, StartStrandID, NumStrands, StartVertexID, NumVertices, Parameters, PropName);
