Groom: Added support for importing groom_group_name attribute from indexed-array.

diff --git a/Engine/Plugins/Importers/AlembicHairImporter/Source/AlembicHairTranslator/Private/AlembicHairTranslator.cpp b/Engine/Plugins/Importers/AlembicHairImporter/Source/AlembicHairTranslator/Private/AlembicHairTranslator.cpp
--- a/Engine/Plugins/Importers/AlembicHairImporter/Source/AlembicHairTranslator/Private/AlembicHairTranslator.cpp
+++ b/Engine/Plugins/Importers/AlembicHairImporter/Source/AlembicHairTranslator/Private/AlembicHairTranslator.cpp
@@ -406,6 +406,83 @@ namespace AlembicHairTranslatorUtils
 		}
 	}
 
+	void ConvertIndexedStringAttribute(FHairDescription& HairDescription, int32 StartStrandID, int32 NumStrands, const Alembic::AbcGeom::ICompoundProperty& CompoundProperty, const std::string& PropName)
+	{
+		// The compound property should contain 2 properties: ".vals" and ".indices"
+		if (CompoundProperty.getNumProperties() != 2)
+		{
+			return;
+		}
+
+		// First property ".vals" should contain the string values
+		TArray<FName> IndexedValues;
+		{
+			Alembic::Abc::PropertyHeader PropertyHeader = CompoundProperty.getPropertyHeader(0);
+			Alembic::Abc::DataType DataType = PropertyHeader.getDataType();
+
+			// Basic property validation: array of single strings
+			if (PropertyHeader.getPropertyType() != Alembic::Abc::kArrayProperty || 
+				DataType.getPod() != Alembic::Util::kStringPOD ||
+				DataType.getExtent() != 1)
+			{
+				return;
+			}
+
+			// The strings are converted and stored temporarily as FName
+			Alembic::AbcGeom::IStringGeomParam Param(CompoundProperty, PropertyHeader.getName());
+			Alembic::Abc::StringArraySamplePtr ParamValues = Param.getExpandedValue().getVals();
+			const int32 NumValues = ParamValues->size();
+
+			IndexedValues.SetNum(NumValues);
+			for (int32 ValIndex = 0; ValIndex < NumValues; ++ValIndex)
+			{
+				IndexedValues[ValIndex] = TConvertor<FName>::ConvertType((*ParamValues)[ValIndex]);
+			}
+		}
+
+		// Second property ".indices" should contain indices for each strand that map to the string values
+		{
+			Alembic::Abc::PropertyHeader PropertyHeader = CompoundProperty.getPropertyHeader(1);
+			Alembic::Abc::DataType DataType = PropertyHeader.getDataType();
+
+			// Basic property validation: array of single uint32 values
+			if (PropertyHeader.getPropertyType() != Alembic::Abc::kArrayProperty ||
+				DataType.getPod() != Alembic::Util::kUint32POD ||
+				DataType.getExtent() != 1)
+			{
+				return;
+			}
+
+			Alembic::AbcGeom::IUInt32GeomParam Param(CompoundProperty, PropertyHeader.getName());
+			Alembic::Abc::UInt32ArraySamplePtr ParamValues = Param.getExpandedValue().getVals();
+			const int32 NumValues = ParamValues->size();
+
+			// Must have an index value for each strand
+			if (NumValues != NumStrands)
+			{
+				return;
+			}
+
+			FName AttributeName(ANSI_TO_TCHAR(PropName.c_str()));
+			TStrandAttributesRef<FName> StrandAttributeRef = HairDescription.StrandAttributes().GetAttributesRef<FName>(AttributeName);
+			if (!StrandAttributeRef.IsValid())
+			{
+				HairDescription.StrandAttributes().RegisterAttribute<FName>(AttributeName);
+				StrandAttributeRef = HairDescription.StrandAttributes().GetAttributesRef<FName>(AttributeName);
+			}
+
+			// The indices are converted to FName and stored in the HairDescription
+			for (int32 StrandIndex = 0; StrandIndex < NumStrands; ++StrandIndex)
+			{
+				uint32 ValueIndex = TConvertor<uint32>::ConvertType((*ParamValues)[StrandIndex]);
+				if (IndexedValues.IsValidIndex(ValueIndex))
+				{
+					StrandAttributeRef[FStrandID(StartStrandID + StrandIndex)] = IndexedValues[ValueIndex];
+				}
+			}
+		}
+	}
+
 	/** Convert the given Alembic parameters to hair attributes in the proper scope */
 	void ConvertAlembicAttributes(FHairDescription& HairDescription, int32 StartStrandID, int32 NumStrands, int32 StartVertexID, int32 NumVertices, const Alembic::AbcGeom::ICompoundProperty& Parameters, FGroomAnimationInfo* AnimInfo)
 	{
@@ -516,6 +593,12 @@ namespace AlembicHairTranslatorUtils
 				break;
 				}
 			}
+			else if (PropType == Alembic::Abc::kCompoundProperty)
+			{
+				// Special case for indexed string attribute
+				Alembic::AbcGeom::ICompoundProperty CompoundProperty(Parameters, PropName);
+				ConvertIndexedStringAttribute(HairDescription, StartStrandID, NumStrands, CompoundProperty, PropName);
+			}
 		}
 	}
 }
