Groom: Added support to extract RootUV values from the ICurves UVsParam instead of the groom_root_uv attribute.

diff --git a/Engine/Plugins/Importers/AlembicHairImporter/Source/AlembicHairTranslator/Private/AlembicHairTranslator.cpp b/Engine/Plugins/Importers/AlembicHairImporter/Source/AlembicHairTranslator/Private/AlembicHairTranslator.cpp
--- a/Engine/Plugins/Importers/AlembicHairImporter/Source/AlembicHairTranslator/Private/AlembicHairTranslator.cpp
+++ b/Engine/Plugins/Importers/AlembicHairImporter/Source/AlembicHairTranslator/Private/AlembicHairTranslator.cpp
@@ -742,6 +742,53 @@ static void ParseObject(const Alembic::Abc::IObject& InObject, float FrameTime,
 			}
 		}
 
+		// Extract RootUV attribute from ICurves UVsParam
+		const bool bHasUVsParam = Curves.getSchema().getUVsParam().valid();
+		if (bHasUVsParam)
+		{
+			Alembic::AbcGeom::IV2fGeomParam UVsParam = Curves.getSchema().getUVsParam();
+			Alembic::Abc::V2fArraySamplePtr UVs = UVsParam.getExpandedValue(SampleSelector).getVals();
+			auto ConvertUVParam = [UVs](int32 Index)
+			{
+				// Convert IV2fGeomParam value to FVector2f by their individual components
+				FVector2f ParamValue;
+				for (int32 ExtentIndex = 0; ExtentIndex < 2; ++ExtentIndex)
+				{
+					ParamValue[ExtentIndex] = (*UVs)[Index][ExtentIndex];
+				}
+				return ParamValue;
+			};
+
+			// Add the RootUV attribute to the HairDescription if needed
+			TStrandAttributesRef<FVector2f> RootUVStrandAttributeRef = HairDescription.StrandAttributes().GetAttributesRef<FVector2f>(HairAttribute::Strand::RootUV);
+			if (!RootUVStrandAttributeRef.IsValid())
+			{
+				HairDescription.StrandAttributes().RegisterAttribute<FVector2f>(HairAttribute::Strand::RootUV);
+				RootUVStrandAttributeRef = HairDescription.StrandAttributes().GetAttributesRef<FVector2f>(HairAttribute::Strand::RootUV);
+			}
+
+			Alembic::AbcGeom::GeometryScope UVScope = UVsParam.getScope();
+			if (UVScope == Alembic::AbcGeom::kUniformScope)
+			{
+				for (uint32 CurveIndex = 0; CurveIndex < NumCurves; ++CurveIndex)
+				{
+					RootUVStrandAttributeRef[FStrandID(StartStrandID + CurveIndex)] = ConvertUVParam(CurveIndex);
+				}
+			}
+			else if (UVScope == Alembic::AbcGeom::kVertexScope)
+			{
+				// RootUV is a strand attribute but UVsParam is vertex-scope so there's a NumVertices to NumCurves conversion needed
+				int32 VertexIndex = 0;
+				for (uint32 CurveIndex = 0; CurveIndex < NumCurves; ++CurveIndex)
+				{
+					RootUVStrandAttributeRef[FStrandID(StartStrandID + CurveIndex)] = ConvertUVParam(VertexIndex);
+
+					// Increment by the number of vertices per curve
+					VertexIndex += (*NumVertices)[CurveIndex];
+				}
+			}
+		}
+
 		// Set width values
 		// Determine the scope of the WidthsParam
 		Alembic::AbcGeom::GeometryScope WidthScope = Alembic::AbcGeom::kUnknownScope;
