Alembic: Fix import issues with corrupted Alembic files.

diff --git a/Engine/Plugins/Importers/AlembicHairImporter/Source/AlembicHairTranslator/Private/AlembicHairTranslator.cpp b/Engine/Plugins/Importers/AlembicHairImporter/Source/AlembicHairTranslator/Private/AlembicHairTranslator.cpp
--- a/Engine/Plugins/Importers/AlembicHairImporter/Source/AlembicHairTranslator/Private/AlembicHairTranslator.cpp
+++ b/Engine/Plugins/Importers/AlembicHairImporter/Source/AlembicHairTranslator/Private/AlembicHairTranslator.cpp
@@ -958,7 +958,7 @@ bool FAlembicHairTranslator::Translate(const FString& FileName, FHairDescription
 	/** Alembic typed root (top) object*/
 	Alembic::Abc::IObject TopObject;
 
-	Factory.setPolicy(Alembic::Abc::ErrorHandler::kThrowPolicy);
+	Factory.setPolicy(Alembic::Abc::ErrorHandler::kQuietNoopPolicy);
 	Factory.setOgawaNumStreams(12);
 
 	// Extract Archive and compression type from file
@@ -1001,7 +1001,7 @@ public:
 	{
 		Alembic::AbcCoreFactory::IFactory::CoreType CompressionType = Alembic::AbcCoreFactory::IFactory::kUnknown;
 
-		Factory.setPolicy(Alembic::Abc::ErrorHandler::kThrowPolicy);
+		Factory.setPolicy(Alembic::Abc::ErrorHandler::kQuietNoopPolicy);
 		Factory.setOgawaNumStreams(12);
 
 		// Extract Archive and compression type from file
@@ -1070,6 +1070,11 @@ FAlembicHairTranslator::~FAlembicHairTranslator()
 
 static void ValidateObject(const Alembic::Abc::IObject& InObject, bool& bHasGeometry, int32& NumCurves)
 {
+	if (!InObject)
+	{
+		return;
+	}
+
 	// Validate that the Alembic has curves only
 	// Any PolyMesh will cause the Alembic to be rejected by this translator
 
@@ -1107,7 +1112,7 @@ bool FAlembicHairTranslator::CanTranslate(const FString& FilePath)
 	/** Alembic typed root (top) object*/
 	Alembic::Abc::IObject TopObject;
 
-	Factory.setPolicy(Alembic::Abc::ErrorHandler::kThrowPolicy);
+	Factory.setPolicy(Alembic::Abc::ErrorHandler::kQuietNoopPolicy);
 	Factory.setOgawaNumStreams(12);
 
 	// Extract Archive and compression type from file
diff --git a/Engine/Plugins/Importers/AlembicImporter/Source/AlembicLibrary/Private/AbcFile.cpp b/Engine/Plugins/Importers/AlembicImporter/Source/AlembicLibrary/Private/AbcFile.cpp
--- a/Engine/Plugins/Importers/AlembicImporter/Source/AlembicLibrary/Private/AbcFile.cpp
+++ b/Engine/Plugins/Importers/AlembicImporter/Source/AlembicLibrary/Private/AbcFile.cpp
@@ -76,7 +76,7 @@ FAbcFile::~FAbcFile()
 
 EAbcImportError FAbcFile::Open()
 {
-	Factory.setPolicy(Alembic::Abc::ErrorHandler::kThrowPolicy);
+	Factory.setPolicy(Alembic::Abc::ErrorHandler::kQuietNoopPolicy);
 	Factory.setOgawaNumStreams(12);
 	
 	// Extract Archive and compression type from file
@@ -427,6 +427,13 @@ EAbcImportError FAbcFile::Import(UAbcImportSettings* InImportSettings)
 
 void FAbcFile::TraverseAbcHierarchy(const Alembic::Abc::IObject& InObject, IAbcObject* InParent)
 {
+	if (!InObject)
+	{
+		TSharedRef<FTokenizedMessage> Message = FTokenizedMessage::Create(EMessageSeverity::Warning, LOCTEXT("AbcInvalidObject", "Invalid object detected: the Alembic file may be corrupted."));
+		FAbcImportLogger::AddImportMessage(Message);
+		return;
+	}
+
 	// Get Header and MetaData info from current Alembic Object
 	Alembic::AbcCoreAbstract::ObjectHeader Header = InObject.getHeader();
 	const Alembic::Abc::MetaData ObjectMetaData = InObject.getMetaData();
@@ -512,13 +519,10 @@ void FAbcFile::TraverseAbcHierarchy(const Alembic::Abc::IObject& InObject, IAbcO
 	}
 
 	// Recursive traversal of child objects
-	if (NumChildren > 0)
+	for (uint32 ChildIndex = 0; ChildIndex < NumChildren; ++ChildIndex)
 	{
-		for (uint32 ChildIndex = 0; ChildIndex < NumChildren; ++ChildIndex)
-		{
-			const Alembic::Abc::IObject& AbcChildObject = InObject.getChild(ChildIndex);
-			TraverseAbcHierarchy(AbcChildObject, CreatedObject);
-		}
+		const Alembic::Abc::IObject& AbcChildObject = InObject.getChild(ChildIndex);
+		TraverseAbcHierarchy(AbcChildObject, CreatedObject);
 	}
 }
 
