Fixed potential uninitialized FVectors in AlembicHairImporter.

diff --git a/Engine/Plugins/Importers/AlembicHairImporter/Source/AlembicHairTranslator/Private/AlembicHairTranslator.cpp b/Engine/Plugins/Importers/AlembicHairImporter/Source/AlembicHairTranslator/Private/AlembicHairTranslator.cpp
--- a/Engine/Plugins/Importers/AlembicHairImporter/Source/AlembicHairTranslator/Private/AlembicHairTranslator.cpp
+++ b/Engine/Plugins/Importers/AlembicHairImporter/Source/AlembicHairTranslator/Private/AlembicHairTranslator.cpp
@@ -103,7 +103,7 @@ namespace AlembicHairTranslatorUtils
 		AbcArraySampleType ParamValues = Param.getExpandedValue().getVals();
 		if (ParamValues->size() == 1)
 		{
-			AttributeType ParamValue;
+			AttributeType ParamValue = {};
 			for (int32 Index = 0; Index < Extent; ++Index)
 			{
 				ParamValue[Index] = (*ParamValues)[0][Index];
@@ -378,7 +378,7 @@ namespace AlembicHairTranslatorUtils
 			bool bIsConstantValue = Scope == Alembic::AbcGeom::kConstantScope && NumValues != NumStrands;
 			for (int32 StrandIndex = 0; StrandIndex < NumStrands; ++StrandIndex)
 			{
-				AttributeType ParamValue;
+				AttributeType ParamValue = {};
 				for (int32 Index = 0; Index < Extent; ++Index)
 				{
 					ParamValue[Index] = (*ParamValues)[bIsConstantValue ? 0 : StrandIndex][Index];
@@ -398,7 +398,7 @@ namespace AlembicHairTranslatorUtils
 
 			for (int32 VertexIndex = 0; VertexIndex < NumVertices; ++VertexIndex)
 			{
-				AttributeType ParamValue;
+				AttributeType ParamValue = {};
 				for (int32 Index = 0; Index < Extent; ++Index)
 				{
 					ParamValue[Index] = (*ParamValues)[VertexIndex][Index];
