PR #10617: AlembicHairImporterFixes: RootUV from Blender Hair / no RootUV registration when not parsing valid data

diff --git a/Engine/Plugins/Importers/AlembicHairImporter/Source/AlembicHairTranslator/Private/AlembicHairTranslator.cpp b/Engine/Plugins/Importers/AlembicHairImporter/Source/AlembicHairTranslator/Private/AlembicHairTranslator.cpp
--- a/Engine/Plugins/Importers/AlembicHairImporter/Source/AlembicHairTranslator/Private/AlembicHairTranslator.cpp
+++ b/Engine/Plugins/Importers/AlembicHairImporter/Source/AlembicHairTranslator/Private/AlembicHairTranslator.cpp
@@ -758,17 +758,23 @@ static void ParseObject(const Alembic::Abc::IObject& InObject, float FrameTime,
 				return ParamValue;
 			};
 
-			// Add the RootUV attribute to the HairDescription if needed
-			TStrandAttributesRef<FVector2f> RootUVStrandAttributeRef = HairDescription.StrandAttributes().GetAttributesRef<FVector2f>(HairAttribute::Strand::RootUV);
-			if (!RootUVStrandAttributeRef.IsValid())
+			auto GetOrCreateRootUVStrandAttributeRef = [&]() 
 			{
-				HairDescription.StrandAttributes().RegisterAttribute<FVector2f>(HairAttribute::Strand::RootUV);
-				RootUVStrandAttributeRef = HairDescription.StrandAttributes().GetAttributesRef<FVector2f>(HairAttribute::Strand::RootUV);
-			}
+				TStrandAttributesRef<FVector2f> AttributeRootUV = HairDescription.StrandAttributes().GetAttributesRef<FVector2f>(HairAttribute::Strand::RootUV);
+			
+				// Add the RootUV attribute to the HairDescription if needed            
+				if (!AttributeRootUV.IsValid())
+				{
+					HairDescription.StrandAttributes().RegisterAttribute<FVector2f>(HairAttribute::Strand::RootUV);
+					AttributeRootUV = HairDescription.StrandAttributes().GetAttributesRef<FVector2f>(HairAttribute::Strand::RootUV);
+				}
+				return AttributeRootUV;                
+			};
 
 			Alembic::AbcGeom::GeometryScope UVScope = UVsParam.getScope();
-			if (UVScope == Alembic::AbcGeom::kUniformScope)
+			if (UVScope == Alembic::AbcGeom::kConstantScope || UVScope == Alembic::AbcGeom::kUniformScope)
 			{
+				TStrandAttributesRef<FVector2f> RootUVStrandAttributeRef = GetOrCreateRootUVStrandAttributeRef();
 				for (uint32 CurveIndex = 0; CurveIndex < NumCurves; ++CurveIndex)
 				{
 					RootUVStrandAttributeRef[FStrandID(StartStrandID + CurveIndex)] = ConvertUVParam(CurveIndex);
@@ -776,6 +782,7 @@ static void ParseObject(const Alembic::Abc::IObject& InObject, float FrameTime,
 			}
 			else if (UVScope == Alembic::AbcGeom::kVertexScope)
 			{
+				TStrandAttributesRef<FVector2f> RootUVStrandAttributeRef = GetOrCreateRootUVStrandAttributeRef();
 				// RootUV is a strand attribute but UVsParam is vertex-scope so there's a NumVertices to NumCurves conversion needed
 				int32 VertexIndex = 0;
 				for (uint32 CurveIndex = 0; CurveIndex < NumCurves; ++CurveIndex)
